Minimal support for achievements on Android Google Play for TappyChicken
Add mostly-stubbed Google Play OSS, enough for WriteAchievements to work through blueprints
Add ability to configure Google Play achievement ID mappings in the editor
Add the achievement ID mappings for TappyChicken

diff --git a/Engine/Build/Android/Java/src/com/epicgames/ue4/GameActivity.java b/Engine/Build/Android/Java/src/com/epicgames/ue4/GameActivity.java
--- a/Engine/Build/Android/Java/src/com/epicgames/ue4/GameActivity.java
+++ b/Engine/Build/Android/Java/src/com/epicgames/ue4/GameActivity.java
@@ -2,6 +2,9 @@
 
 package com.epicgames.ue4;
 
+import java.util.Map;
+import java.util.HashMap;
+
 import android.app.NativeActivity;
 import android.os.Bundle;
 import android.util.Log;
@@ -15,9 +18,14 @@ import android.content.res.Configuration;
 import android.content.IntentSender.SendIntentException;
 
 import com.google.android.gms.common.api.GoogleApiClient;
+import com.google.android.gms.common.api.PendingResult;
+import com.google.android.gms.common.api.ResultCallback;
 import com.google.android.gms.common.ConnectionResult;
-import com.google.android.gms.games.Games;
 import com.google.android.gms.common.GooglePlayServicesUtil;
+import com.google.android.gms.games.achievement.*;
+import com.google.android.gms.games.Games;
+
+
 
 // TODO: use the resources from the UE4 lib project once we've got the packager up and running
 //import com.epicgames.ue4.R;
@@ -53,14 +61,43 @@ public class GameActivity extends NativeActivity implements GoogleApiClient.Conn
 	/** Unique ID to identify Google Play Services error dialog */
 	private static final int PLAY_SERVICES_DIALOG_ID = 1;
 
+	/** Arbitrary ID for leaderboard display */
+	private static final int REQUEST_LEADERBOARDS = 0;
+	
+	/** Arbitrary ID for achievement display */
+	private static final int REQUEST_ACHIEVEMENTS = 1;
+
+	/** Stores the minimum amount of data we need to set achievement progress */
+	private class BasicAchievementData
+	{
+		public BasicAchievementData()
+		{
+			Type = Achievement.TYPE_STANDARD;
+			MaxSteps = 1;
+		}
+
+		public BasicAchievementData(int InMaxSteps)
+		{
+			Type = Achievement.TYPE_INCREMENTAL;
+			MaxSteps = InMaxSteps;
+		}
+
+		public int Type;
+		public int MaxSteps;
+	}
+
+	/**
+	 * Store achievement data upon login so that we can convert the percentage values from the game to
+	 * integer steps for Google Play
+	 */
+	private Map<String, BasicAchievementData> CachedAchievements = new HashMap<String, BasicAchievementData>();
+
 	@Override
 	public void onStart()
 	{
 		super.onStart();
 		
 		Log.debug("==================================> Inside onStart function in GameActivity");
-		
-		AndroidThunkJava_GooglePlayConnect();
 	}
 
 	@Override
@@ -136,11 +173,48 @@ public class GameActivity extends NativeActivity implements GoogleApiClient.Conn
 		googleClient.disconnect();
 	}
 
+	/** Callback that fills in CachedAchievements when the load operation completes. */
+	private class AchievementsResultStartupCallback implements ResultCallback<Achievements.LoadAchievementsResult>
+	{
+		@Override
+		public void onResult(Achievements.LoadAchievementsResult result)
+		{
+			Log.debug("Google Play Services: Loaded achievements with status " + result.getStatus().toString());
+
+			AchievementBuffer Achievements = result.getAchievements();
+
+			CachedAchievements.clear();
+			for(int i = 0; i < Achievements.getCount(); ++i)
+			{
+				Achievement CurrentAchievement = Achievements.get(i);
+
+				if(CurrentAchievement.getType() == Achievement.TYPE_STANDARD)
+				{
+					CachedAchievements.put(new String(CurrentAchievement.getAchievementId()), new BasicAchievementData());
+				}
+				else if(CurrentAchievement.getType() == Achievement.TYPE_INCREMENTAL)
+				{
+					CachedAchievements.put(new String(CurrentAchievement.getAchievementId()),
+						new BasicAchievementData(CurrentAchievement.getTotalSteps()));
+				}
+			}
+
+			Achievements.close();
+			result.release();
+		}
+	}
+
 	// Callbacks to handle connections with Google Play
-	 @Override
+	@Override
     public void onConnected(Bundle connectionHint)
 	{
         Log.debug("Connected to Google Play Services.");
+
+		// Load achievements. Since games are expected to pass in achievement progress as a percentage,
+		// we need to know what the maximum steps are in order to convert the percentage to an integer
+		// number of steps.
+		PendingResult<Achievements.LoadAchievementsResult> loadAchievementsResult = Games.Achievements.load(googleClient, false);
+		loadAchievementsResult.setResultCallback(new AchievementsResultStartupCallback());
     }
 
     @Override
@@ -254,28 +328,31 @@ public class GameActivity extends NativeActivity implements GoogleApiClient.Conn
 
 	public void AndroidThunkJava_GooglePlayConnect()
 	{
-		if(nativeIsGooglePlayEnabled())
-		{
-			googleClient.connect();
-		}
+		googleClient.connect();
 	}
 
 	public void AndroidThunkJava_ShowLeaderboard(String LeaderboardID)
 	{
 		Log.debug("In AndroidThunkJava_ShowLeaderboard, ID is " + LeaderboardID);
-		if(googleClient.isConnected())
+		if(!googleClient.isConnected())
 		{
-			startActivityForResult(Games.Leaderboards.getLeaderboardIntent(googleClient, LeaderboardID), 0);
+			Log.debug("Not connected to Google Play, can't show leaderboards UI.");
+			return;
 		}
+
+		startActivityForResult(Games.Leaderboards.getLeaderboardIntent(googleClient, LeaderboardID), REQUEST_LEADERBOARDS);
 	}
 
 	public void AndroidThunkJava_ShowAchievements()
 	{
 		Log.debug("In AndroidThunkJava_ShowAchievements");
-		if(googleClient.isConnected())
+		if(!googleClient.isConnected())
 		{
-			startActivityForResult(Games.Achievements.getAchievementsIntent(googleClient), 0);
+			Log.debug("Not connected to Google Play, can't show achievements UI.");
+			return;
 		}
+		
+		startActivityForResult(Games.Achievements.getAchievementsIntent(googleClient), REQUEST_ACHIEVEMENTS);
 	}
 
 	public void AndroidThunkJava_WriteLeaderboardValue(String LeaderboardID, long Value)
@@ -287,6 +364,55 @@ public class GameActivity extends NativeActivity implements GoogleApiClient.Conn
 		}
 	}
 
+	public void AndroidThunkJava_WriteAchievement(String AchievementID, float Percentage)
+	{
+		BasicAchievementData Data = CachedAchievements.get(AchievementID);
+
+		if(Data == null)
+		{
+			Log.debug("Couldn't find cached achievement for ID " + AchievementID + ", not setting progress.");
+			return;
+		}
+
+		if(!googleClient.isConnected())
+		{
+			Log.debug("Not connected to Google Play, can't set achievement progress.");
+			return;
+		}
+
+		// Found the one to unlock.
+		switch(Data.Type)
+		{
+			case Achievement.TYPE_INCREMENTAL:
+			{
+				float StepFraction = (Percentage / 100.0f) * Data.MaxSteps;
+				int RoundedSteps = Math.round(StepFraction);
+
+				if(RoundedSteps > 0)
+				{
+					Log.debug("Incremental achievement ID " + AchievementID + ": setting progress to " + RoundedSteps);
+					Games.Achievements.setSteps(googleClient, AchievementID, RoundedSteps);
+				}
+				else
+				{
+					Log.debug("Incremental achievement ID " + AchievementID + ": not setting progress to " + RoundedSteps);
+				}
+				break;
+			}
+
+			case Achievement.TYPE_STANDARD:
+			{
+				// Standard achievements only unlock if the progress is at least 100%.
+				if(Percentage >= 100.0f)
+				{
+					Log.debug("Standard achievement ID " + AchievementID + ": unlocking");
+					Games.Achievements.unlock(googleClient, AchievementID);
+				}
+				break;
+			}
+		}
+	}
+
 	public native boolean nativeIsShippingBuild();
 	public native void nativeSetGlobalActivity();
 	public native void nativeSetWindowInfo(boolean bIsPortrait);
@@ -294,8 +420,6 @@ public class GameActivity extends NativeActivity implements GoogleApiClient.Conn
 
 	public native void nativeConsoleCommand(String commandString);
 	
-	public native boolean nativeIsGooglePlayEnabled();
-
 	static
 	{
 		System.loadLibrary("UE4");
diff --git a/Engine/Config/Android/AndroidEngine.ini b/Engine/Config/Android/AndroidEngine.ini
--- a/Engine/Config/Android/AndroidEngine.ini
+++ b/Engine/Config/Android/AndroidEngine.ini
@@ -9,3 +9,6 @@ DeviceProfileSelectionModule="AndroidDeviceProfileSelector"
 
 [Advertising]
 DefaultProviderName=AndroidAdvertising
+
+[OnlineSubsystem]
+DefaultPlatformService=GooglePlay
diff --git a/Engine/Source/Runtime/Android/AndroidRuntimeSettings/Classes/AndroidRuntimeSettings.h b/Engine/Source/Runtime/Android/AndroidRuntimeSettings/Classes/AndroidRuntimeSettings.h
--- a/Engine/Source/Runtime/Android/AndroidRuntimeSettings/Classes/AndroidRuntimeSettings.h
+++ b/Engine/Source/Runtime/Android/AndroidRuntimeSettings/Classes/AndroidRuntimeSettings.h
@@ -35,6 +35,23 @@ namespace EAndroidScreenOrientation
 	};
 }
 
+/**
+ * Holds the game-specific achievement name and corresponding ID from Google Play services.
+ */
+USTRUCT()
+struct FGooglePlayAchievementMapping
+{
+	GENERATED_USTRUCT_BODY()
+
+	// The game-specific achievement name (the one passed in to WriteAchievement calls)
+	UPROPERTY(EditAnywhere, Category = GooglePlayServices)
+	FString Name;
+
+	// The ID of the corresponding achievement, generated by the Google Play developer console.
+	UPROPERTY(EditAnywhere, Category = GooglePlayServices)
+	FString AchievementID;
+};
+
 /**
  * Implements the settings for the Android runtime platform.
  */
@@ -55,4 +72,8 @@ public:
 	// The app id obtained from the Google Play Developer Console
 	UPROPERTY(GlobalConfig, EditAnywhere, Category = GooglePlayServices)
 	FString GamesAppID;
+
+	// Mapping of game achievement names to IDs generated by Google Play.
+	UPROPERTY(GlobalConfig, EditAnywhere, Category = GooglePlayServices)
+	TArray<FGooglePlayAchievementMapping> AchievementMap;
 };
diff --git a/Engine/Source/Runtime/CoreUObject/Public/UObject/UObjectGlobals.h b/Engine/Source/Runtime/CoreUObject/Public/UObject/UObjectGlobals.h
--- a/Engine/Source/Runtime/CoreUObject/Public/UObject/UObjectGlobals.h
+++ b/Engine/Source/Runtime/CoreUObject/Public/UObject/UObjectGlobals.h
@@ -7,6 +7,8 @@
 #ifndef __UNOBJGLOBALS_H__
 #define __UNOBJGLOBALS_H__
 
+#include "Script.h"
+
 COREUOBJECT_API DECLARE_LOG_CATEGORY_EXTERN(LogUObjectGlobals, Log, All);
 
 DECLARE_CYCLE_STAT_EXTERN(TEXT("ConstructObject"),STAT_ConstructObject,STATGROUP_Object, );
diff --git a/Engine/Source/Runtime/Launch/Private/Android/AndroidJNI.cpp b/Engine/Source/Runtime/Launch/Private/Android/AndroidJNI.cpp
--- a/Engine/Source/Runtime/Launch/Private/Android/AndroidJNI.cpp
+++ b/Engine/Source/Runtime/Launch/Private/Android/AndroidJNI.cpp
@@ -98,6 +98,7 @@ jmethodID JDef_GameActivity::AndroidThunkJava_ShowLeaderboard;
 jmethodID JDef_GameActivity::AndroidThunkJava_ShowAchievements;
 jmethodID JDef_GameActivity::AndroidThunkJava_WriteLeaderboardValue;
 jmethodID JDef_GameActivity::AndroidThunkJava_GooglePlayConnect;
+jmethodID JDef_GameActivity::AndroidThunkJava_WriteAchievement;
 
 DEFINE_LOG_CATEGORY_STATIC(LogEngine, Log, All);
 
@@ -193,6 +194,23 @@ void AndroidThunkCpp_WriteLeaderboardValue(const FString& LeaderboardName, int64
 	}
 }
 
+void AndroidThunkCpp_GooglePlayConnect()
+{
+	if (JNIEnv* Env = GetJavaEnv())
+	{
+		Env->CallVoidMethod(GJavaGlobalThis, JDef_GameActivity::AndroidThunkJava_GooglePlayConnect);
+	}
+}
+
+void AndroidThunkCpp_WriteAchievement(const FString& AchievementID, float PercentComplete)
+{
+	if (JNIEnv* Env = GetJavaEnv())
+	{
+		jstring AchievementIDArg = Env->NewStringUTF(TCHAR_TO_UTF8(*AchievementID));
+		Env->CallVoidMethod(GJavaGlobalThis, JDef_GameActivity::AndroidThunkJava_WriteAchievement, AchievementIDArg, PercentComplete);
+		Env->DeleteLocalRef(AchievementIDArg);
+	}
+}
 
 //The JNI_OnLoad function is triggered by loading the game library from 
 //the Java source file.
@@ -221,6 +239,7 @@ JNIEXPORT jint JNI_OnLoad(JavaVM* InJavaVM, void* InReserved)
 	JDef_GameActivity::AndroidThunkJava_ShowAchievements = env->GetMethodID(JDef_GameActivity::ClassID, "AndroidThunkJava_ShowAchievements", "()V");
 	JDef_GameActivity::AndroidThunkJava_WriteLeaderboardValue = env->GetMethodID(JDef_GameActivity::ClassID, "AndroidThunkJava_WriteLeaderboardValue", "(Ljava/lang/String;J)V");
 	JDef_GameActivity::AndroidThunkJava_GooglePlayConnect = env->GetMethodID(JDef_GameActivity::ClassID, "AndroidThunkJava_GooglePlayConnect", "()V");
+	JDef_GameActivity::AndroidThunkJava_WriteAchievement = env->GetMethodID(JDef_GameActivity::ClassID, "AndroidThunkJava_WriteAchievement", "(Ljava/lang/String;F)V");
 
 	// hook signals
 #if UE_BUILD_DEBUG
diff --git a/Engine/Source/Runtime/Launch/Private/Android/AndroidJNI.h b/Engine/Source/Runtime/Launch/Private/Android/AndroidJNI.h
--- a/Engine/Source/Runtime/Launch/Private/Android/AndroidJNI.h
+++ b/Engine/Source/Runtime/Launch/Private/Android/AndroidJNI.h
@@ -25,6 +25,7 @@ public:
 	static jmethodID AndroidThunkJava_ShowAchievements;
 	static jmethodID AndroidThunkJava_WriteLeaderboardValue;
 	static jmethodID AndroidThunkJava_GooglePlayConnect;
+	static jmethodID AndroidThunkJava_WriteAchievement;
 };
 
 // Returns the java environment
diff --git a/Engine/Source/Runtime/Online/Android/OnlineSubsystemGooglePlay/OnlineSubsystemGooglePlay.Build.cs b/Engine/Source/Runtime/Online/Android/OnlineSubsystemGooglePlay/OnlineSubsystemGooglePlay.Build.cs
new file mode 100644
--- /dev/null
+++ b/Engine/Source/Runtime/Online/Android/OnlineSubsystemGooglePlay/OnlineSubsystemGooglePlay.Build.cs
@@ -0,0 +1,30 @@
+// Copyright 1998-2014 Epic Games, Inc. All Rights Reserved.
+
+using UnrealBuildTool;
+
+public class OnlineSubsystemGooglePlay : ModuleRules
+{
+	public OnlineSubsystemGooglePlay(TargetInfo Target)
+	{
+		Definitions.Add("ONLINESUBSYSTEMGOOGLEPLAY_PACKAGE=1");
+
+		PrivateIncludePaths.AddRange(
+			new string[] {
+				"Runtime/Online/Android/OnlineSubsystemGooglePlay/Private",    
+				"Runtime/Launch/Private"
+			}
+			);
+
+		PrivateDependencyModuleNames.AddRange(
+			new string[] { 
+                "Core", 
+                "CoreUObject", 
+                "Engine", 
+                "Sockets",
+				"OnlineSubsystem", 
+                "Http",
+				"AndroidRuntimeSettings"
+            }
+			);
+	}
+}
diff --git a/Engine/Source/Runtime/Online/Android/OnlineSubsystemGooglePlay/Private/OnlineAchievementsInterfaceGooglePlay.cpp b/Engine/Source/Runtime/Online/Android/OnlineSubsystemGooglePlay/Private/OnlineAchievementsInterfaceGooglePlay.cpp
new file mode 100644
--- /dev/null
+++ b/Engine/Source/Runtime/Online/Android/OnlineSubsystemGooglePlay/Private/OnlineAchievementsInterfaceGooglePlay.cpp
@@ -0,0 +1,154 @@
+// Copyright 1998-2014 Epic Games, Inc. All Rights Reserved.
+
+// Module includes
+#include "OnlineSubsystemGooglePlayPrivatePCH.h"
+
+#include <jni.h>
+
+FOnlineAchievementsGooglePlay::FOnlineAchievementsGooglePlay( FOnlineSubsystemGooglePlay* InSubsystem )
+	: CurrentWriteObject(MakeShareable(new FOnlineAchievementsWrite()))
+{
+	AndroidSubsystem = InSubsystem;
+	CurrentPlayerID = MakeShareable( new FUniqueNetIdString(FString("")));
+	
+	TArray<FString> MapAsStrings;
+
+	// Use the config to find the Google Play-assigned achievement ID corresponding to the name the game passed in.
+	GConfig->GetArray(TEXT("/Script/AndroidRuntimeSettings.AndroidRuntimeSettings"), TEXT("AchievementMap"), MapAsStrings, GEngineIni);
+
+	for(int32 x = 0; x < MapAsStrings.Num(); ++x)
+	{
+		FGooglePlayAchievementMapping Mapping;
+
+		if(!FParse::Value(*MapAsStrings[x], TEXT("Name="), Mapping.Name))
+		{
+			// Could not find an achievement name, cannot continue with this line
+			continue;
+		}
+
+		if(!FParse::Value(*MapAsStrings[x], TEXT("AchievementID="), Mapping.AchievementID))
+		{
+			// Could not find an achievement id, cannot continue with this line
+			continue;
+		}
+
+		UE_LOG(LogOnline, Log, TEXT("Found achievement mapping: %s to %s"), *Mapping.Name, *Mapping.AchievementID);
+		AchievementConfigMap.Add(Mapping);
+	}
+}
+
+void FOnlineAchievementsGooglePlay::SetUniqueID(const FUniqueNetId& InUserID)
+{
+	CurrentPlayerID = MakeShareable( new FUniqueNetIdString(InUserID));
+}
+
+void FOnlineAchievementsGooglePlay::ResetUniqueID()
+{
+	CurrentPlayerID.Reset();
+}
+
+void FOnlineAchievementsGooglePlay::QueryAchievements(const FUniqueNetId& PlayerId, const FOnQueryAchievementsCompleteDelegate& Delegate)
+{
+	Delegate.ExecuteIfBound( PlayerId, false );
+	return;
+}
+
+
+void FOnlineAchievementsGooglePlay::QueryAchievementDescriptions(const FUniqueNetId& PlayerId, const FOnQueryAchievementsCompleteDelegate& Delegate)
+{
+	// Just query achievements to get descriptions
+	// FIXME: This feels a little redundant, but we can see how platforms evolve, and make a decision then
+	QueryAchievements( PlayerId, Delegate );
+}
+
+void FOnlineAchievementsGooglePlay::WriteAchievements( const FUniqueNetId& PlayerId, FOnlineAchievementsWriteRef& WriteObject, const FOnAchievementsWrittenDelegate& Delegate )
+{
+	for (FStatPropertyArray::TConstIterator It(WriteObject->Properties); It; ++It)
+	{
+		// Access the stat and the value.
+		const FVariantData& Stat = It.Value();
+
+		// Create an achievement object which should be reported to the server.
+		const FString AchievementName(It.Key().ToString());
+		float PercentComplete = 0.0f;
+
+		// Setup the percentage complete with the value we are writing from the variant type
+		switch (Stat.GetType())
+		{
+		case EOnlineKeyValuePairDataType::Int32:
+			{
+				int32 Value;
+				Stat.GetValue(Value);
+				PercentComplete = (float)Value;
+				break;
+			}
+		case EOnlineKeyValuePairDataType::Float:
+			{
+				float Value;
+				Stat.GetValue(Value);
+				PercentComplete = Value;
+				break;
+			}
+		default:
+			{
+				UE_LOG(LogOnline, Error, TEXT("FOnlineAchievementsGooglePlay Trying to write an achievement with incompatible format. Not a float or int"));
+				break;
+			}
+		}
+
+		// Find the ID in the mapping
+		for(const auto& Mapping : AchievementConfigMap)
+		{
+			if(Mapping.Name == AchievementName)
+			{
+				UE_LOG(LogOnline, Log, TEXT("Writing achievement name: %s, id: %s, progress: %.0f"),
+					*Mapping.Name, *Mapping.AchievementID, PercentComplete);
+				extern void AndroidThunkCpp_WriteAchievement(const FString& AchievementID, float PercentComplete);
+				AndroidThunkCpp_WriteAchievement(Mapping.AchievementID, PercentComplete);
+				break;
+			}
+		}
+	}
+}
+
+EOnlineCachedResult::Type FOnlineAchievementsGooglePlay::GetCachedAchievements(const FUniqueNetId& PlayerId, TArray<FOnlineAchievement>& OutAchievements)
+{
+	// look up achievements for player
+	OutAchievements = Achievements;
+
+	// did we have them cached?
+	return OutAchievements.Num() > 0 ? EOnlineCachedResult::Success : EOnlineCachedResult::NotFound;
+}
+
+
+EOnlineCachedResult::Type FOnlineAchievementsGooglePlay::GetCachedAchievementDescription(const FString& AchievementId, FOnlineAchievementDesc& OutAchievementDesc)
+{
+	FOnlineAchievementDesc* FoundDesc = AchievementDescriptions.Find( AchievementId );
+	if( FoundDesc != NULL )
+	{
+		OutAchievementDesc = *FoundDesc;
+	}
+	return FoundDesc != NULL ? EOnlineCachedResult::Success : EOnlineCachedResult::NotFound;
+}
+
+#if !UE_BUILD_SHIPPING
+bool FOnlineAchievementsGooglePlay::ResetAchievements( const FUniqueNetId& PlayerId )
+{
+	check(!TEXT("ResetAchievements has not been implemented"));
+	return false;
+};
+#endif // !UE_BUILD_SHIPPING
+
+EOnlineCachedResult::Type FOnlineAchievementsGooglePlay::GetCachedAchievement( const FUniqueNetId& PlayerId, const FString& AchievementId, FOnlineAchievement& OutAchievement )
+{
+	for(auto& CurrentAchievement : Achievements)
+	{
+		if(CurrentAchievement.Id == AchievementId)
+		{
+			OutAchievement = CurrentAchievement;
+			return EOnlineCachedResult::Success;
+		}
+	}
+
+	return EOnlineCachedResult::NotFound;
+}
diff --git a/Engine/Source/Runtime/Online/Android/OnlineSubsystemGooglePlay/Private/OnlineAchievementsInterfaceGooglePlay.h b/Engine/Source/Runtime/Online/Android/OnlineSubsystemGooglePlay/Private/OnlineAchievementsInterfaceGooglePlay.h
new file mode 100644
--- /dev/null
+++ b/Engine/Source/Runtime/Online/Android/OnlineSubsystemGooglePlay/Private/OnlineAchievementsInterfaceGooglePlay.h
@@ -0,0 +1,73 @@
+// Copyright 1998-2014 Epic Games, Inc. All Rights Reserved.
+
+#pragma once
+
+#include "OnlineAchievementsInterface.h"
+#include "OnlineSubsystemTypes.h"
+#include "AndroidRuntimeSettings.h"
+
+/**
+ *	IOnlineAchievements - Interface class for Achievements
+ */
+class FOnlineAchievementsGooglePlay : public IOnlineAchievements
+{
+private:
+
+	/** Reference to the main GameCenter subsystem */
+	class FOnlineSubsystemGooglePlay* AndroidSubsystem;
+	
+	/** hide the default constructor, we need a reference to our OSS */
+	FOnlineAchievementsGooglePlay() : CurrentWriteObject(MakeShareable(new FOnlineAchievementsWrite()))
+	{}
+
+	// Android only supports loading achievements for local player. This is where they are cached.
+	TArray< FOnlineAchievement > Achievements;
+
+	// Cached achievement descriptions for an Id
+	TMap< FString, FOnlineAchievementDesc > AchievementDescriptions;
+
+public:
+
+	// Begin IOnlineAchievements interface
+	virtual void WriteAchievements(const FUniqueNetId& PlayerId, FOnlineAchievementsWriteRef& WriteObject, const FOnAchievementsWrittenDelegate& Delegate = FOnAchievementsWrittenDelegate()) OVERRIDE;
+	virtual void QueryAchievements(const FUniqueNetId& PlayerId, const FOnQueryAchievementsCompleteDelegate& Delegate = FOnQueryAchievementsCompleteDelegate()) OVERRIDE;
+	virtual void QueryAchievementDescriptions(const FUniqueNetId& PlayerId, const FOnQueryAchievementsCompleteDelegate& Delegate = FOnQueryAchievementsCompleteDelegate()) OVERRIDE ;
+	virtual EOnlineCachedResult::Type GetCachedAchievement(const FUniqueNetId& PlayerId, const FString& AchievementId, FOnlineAchievement& OutAchievement) OVERRIDE;
+	virtual EOnlineCachedResult::Type GetCachedAchievements(const FUniqueNetId& PlayerId, TArray<FOnlineAchievement>& OutAchievements) OVERRIDE;
+	virtual EOnlineCachedResult::Type GetCachedAchievementDescription(const FString& AchievementId, FOnlineAchievementDesc& OutAchievementDesc) OVERRIDE;
+#if !UE_BUILD_SHIPPING
+	virtual bool ResetAchievements( const FUniqueNetId& PlayerId ) OVERRIDE;
+#endif // !UE_BUILD_SHIPPING
+	// End IOnlineAchievements interface
+	
+	//platform specific members
+
+	/** Sets unique ID of current logged-in player */
+	void SetUniqueID(const FUniqueNetId& InUserID);
+	
+	/** Reset unique ID of current logged-in player */
+	void ResetUniqueID();
+
+	/**
+	 * Constructor
+	 *
+	 * @param InSubsystem - A reference to the owning subsystem
+	 */
+	FOnlineAchievementsGooglePlay( class FOnlineSubsystemGooglePlay* InSubsystem );
+
+	/**
+	 * Default destructor
+	 */
+	virtual ~FOnlineAchievementsGooglePlay() {}
+
+
+private:
+
+	TArray<FGooglePlayAchievementMapping> AchievementConfigMap;
+
+	TSharedPtr< FUniqueNetIdString > CurrentPlayerID;
+
+	TSharedRef<FOnlineAchievementsWrite> CurrentWriteObject;
+};
+
+typedef TSharedPtr<FOnlineAchievementsGooglePlay, ESPMode::ThreadSafe> FOnlineAchievementsGooglePlayPtr;
diff --git a/Engine/Source/Runtime/Online/Android/OnlineSubsystemGooglePlay/Private/OnlineIdentityInterfaceGooglePlay.h b/Engine/Source/Runtime/Online/Android/OnlineSubsystemGooglePlay/Private/OnlineIdentityInterfaceGooglePlay.h
new file mode 100644
--- /dev/null
+++ b/Engine/Source/Runtime/Online/Android/OnlineSubsystemGooglePlay/Private/OnlineIdentityInterfaceGooglePlay.h
@@ -0,0 +1,59 @@
+// Copyright 1998-2014 Epic Games, Inc. All Rights Reserved.
+
+#pragma once
+
+#include "OnlineIdentityInterface.h"
+#include "OnlineSubsystemGooglePlayPackage.h"
+
+class FOnlineIdentityGooglePlay :
+	public IOnlineIdentity
+{
+private:
+	
+	bool bPrevLoggedIn;
+	bool bLoggedIn;
+	FString PlayerAlias;
+	int32 CurrentLoginUserNum;
+	class FOnlineSubsystemGooglePlay* MainSubsystem;
+
+	bool bLoggingInUser;
+	bool bRegisteringUser;
+
+PACKAGE_SCOPE:
+
+	/**
+	 * Default Constructor
+	 */
+	FOnlineIdentityGooglePlay(FOnlineSubsystemGooglePlay* InSubsystem);
+
+public:
+
+	// Begin IOnlineIdentity interface
+	virtual TSharedPtr<FUserOnlineAccount> GetUserAccount(const FUniqueNetId& UserId) const OVERRIDE;
+	virtual TArray<TSharedPtr<FUserOnlineAccount> > GetAllUserAccounts() const OVERRIDE;
+	virtual bool Login(int32 LocalUserNum, const FOnlineAccountCredentials& AccountCredentials) OVERRIDE;
+	virtual bool Logout(int32 LocalUserNum) OVERRIDE;
+	virtual bool AutoLogin(int32 LocalUserNum) OVERRIDE;
+	virtual TSharedPtr<FUniqueNetId> GetUniquePlayerId(int32 LocalUserNum) const OVERRIDE;
+	virtual TSharedPtr<FUniqueNetId> CreateUniquePlayerId(uint8* Bytes, int32 Size) OVERRIDE;
+	virtual TSharedPtr<FUniqueNetId> CreateUniquePlayerId(const FString& Str) OVERRIDE;
+	virtual ELoginStatus::Type GetLoginStatus(int32 LocalUserNum) const OVERRIDE;
+	virtual FString GetPlayerNickname(int32 LocalUserNum) const OVERRIDE;
+	virtual FString GetAuthToken(int32 LocalUserNum) const OVERRIDE;
+	// End IOnlineIdentity interface
+
+public:
+
+	// Begin IOnlineIdentity interface
+	void Tick(float DeltaTime);
+	// End IOnlineIdentity interface
+	
+	//platform specific
+	void OnLogin(const bool InLoggedIn,  const FString& InPlayerId, const FString& InPlayerAlias);
+	void OnLogout(const bool InLoggedIn);
+
+};
+
+
+typedef TSharedPtr<FOnlineIdentityGooglePlay, ESPMode::ThreadSafe> FOnlineIdentityGooglePlayPtr;
+
diff --git a/Engine/Source/Runtime/Online/Android/OnlineSubsystemGooglePlay/Private/OnlineIdentityInterfaceGoolgePlay.cpp b/Engine/Source/Runtime/Online/Android/OnlineSubsystemGooglePlay/Private/OnlineIdentityInterfaceGoolgePlay.cpp
new file mode 100644
--- /dev/null
+++ b/Engine/Source/Runtime/Online/Android/OnlineSubsystemGooglePlay/Private/OnlineIdentityInterfaceGoolgePlay.cpp
@@ -0,0 +1,113 @@
+// Copyright 1998-2014 Epic Games, Inc. All Rights Reserved.
+
+#include "OnlineSubsystemGooglePlayPrivatePCH.h"
+
+FOnlineIdentityGooglePlay::FOnlineIdentityGooglePlay(FOnlineSubsystemGooglePlay* InSubsystem)
+	: bPrevLoggedIn(false)
+	, bLoggedIn(false)
+	, PlayerAlias("NONE")
+	, MainSubsystem(InSubsystem)
+	, bRegisteringUser(false)
+	, bLoggingInUser(false)
+
+{
+	UE_LOG(LogOnline, Display, TEXT("FOnlineIdentityAndroid::FOnlineIdentityAndroid()"));
+}
+
+TSharedPtr<FUserOnlineAccount> FOnlineIdentityGooglePlay::GetUserAccount(const FUniqueNetId& UserId) const
+{
+	return nullptr;
+}
+
+
+TArray<TSharedPtr<FUserOnlineAccount> > FOnlineIdentityGooglePlay::GetAllUserAccounts() const
+{
+	TArray<TSharedPtr<FUserOnlineAccount> > Result;
+
+	return Result;
+}
+
+
+bool FOnlineIdentityGooglePlay::Login(int32 LocalUserNum, const FOnlineAccountCredentials& AccountCredentials) 
+{
+	TriggerOnLoginCompleteDelegates(LocalUserNum, false, FUniqueNetIdString(TEXT("")), FString("Login not supported"));
+	return false;
+}
+
+
+bool FOnlineIdentityGooglePlay::Logout(int32 LocalUserNum)
+{
+	TriggerOnLogoutCompleteDelegates(LocalUserNum, false);
+	return false;
+}
+
+
+bool FOnlineIdentityGooglePlay::AutoLogin(int32 LocalUserNum)
+{
+	return false;
+}
+
+
+ELoginStatus::Type FOnlineIdentityGooglePlay::GetLoginStatus(int32 LocalUserNum) const
+{
+	ELoginStatus::Type LoginStatus = ELoginStatus::NotLoggedIn;
+
+	if(LocalUserNum < MAX_LOCAL_PLAYERS && bLoggedIn)
+	{
+		UE_LOG(LogOnline, Display, TEXT("FOnlineIdentityAndroid::GetLoginStatus - Logged in"));
+		LoginStatus = ELoginStatus::LoggedIn;
+	}
+	else
+	{
+		UE_LOG(LogOnline, Display, TEXT("FOnlineIdentityAndroid::GetLoginStatus - Not logged in"));
+	}
+
+	return LoginStatus;
+}
+
+
+TSharedPtr<FUniqueNetId> FOnlineIdentityGooglePlay::GetUniquePlayerId(int32 LocalUserNum) const
+{
+	TSharedPtr<FUniqueNetId> NewID = MakeShareable(new FUniqueNetIdString(""));
+	return NewID;
+}
+
+
+TSharedPtr<FUniqueNetId> FOnlineIdentityGooglePlay::CreateUniquePlayerId(uint8* Bytes, int32 Size)
+{
+	if( Bytes && Size == sizeof(uint64) )
+	{
+		int32 StrLen = FCString::Strlen((TCHAR*)Bytes);
+		if (StrLen > 0)
+		{
+			FString StrId((TCHAR*)Bytes);
+			return MakeShareable(new FUniqueNetIdString(StrId));
+		}
+	}
+	return NULL;
+}
+
+
+TSharedPtr<FUniqueNetId> FOnlineIdentityGooglePlay::CreateUniquePlayerId(const FString& Str)
+{
+	return MakeShareable(new FUniqueNetIdString(Str));
+}
+
+
+FString FOnlineIdentityGooglePlay::GetPlayerNickname(int32 LocalUserNum) const
+{
+	UE_LOG(LogOnline, Display, TEXT("FOnlineIdentityAndroid::GetPlayerNickname"));
+	return PlayerAlias;
+}
+
+
+FString FOnlineIdentityGooglePlay::GetAuthToken(int32 LocalUserNum) const
+{
+	UE_LOG(LogOnline, Display, TEXT("FOnlineIdentityAndroid::GetAuthToken not implemented"));
+	FString ResultToken;
+	return ResultToken;
+}
+
+void FOnlineIdentityGooglePlay::Tick(float DeltaTime)
+{
+}
diff --git a/Engine/Source/Runtime/Online/Android/OnlineSubsystemGooglePlay/Private/OnlineSubsystemGooglePlay.cpp b/Engine/Source/Runtime/Online/Android/OnlineSubsystemGooglePlay/Private/OnlineSubsystemGooglePlay.cpp
new file mode 100644
--- /dev/null
+++ b/Engine/Source/Runtime/Online/Android/OnlineSubsystemGooglePlay/Private/OnlineSubsystemGooglePlay.cpp
@@ -0,0 +1,148 @@
+// Copyright 1998-2014 Epic Games, Inc. All Rights Reserved.
+
+//Google Play Services
+
+#include "OnlineSubsystemGooglePlayPrivatePCH.h"
+#include "AndroidApplication.h"
+
+FOnlineSubsystemGooglePlay::FOnlineSubsystemGooglePlay()
+	: OnlineAsyncTaskThread(NULL)
+	, IdentityInterface(NULL)
+	, AchievementsInterface(NULL)
+{
+
+}
+
+IOnlineIdentityPtr FOnlineSubsystemGooglePlay::GetIdentityInterface() const
+{
+	return IdentityInterface;
+}
+
+
+IOnlineStorePtr FOnlineSubsystemGooglePlay::GetStoreInterface() const
+{
+	return nullptr;
+}
+
+
+IOnlineSessionPtr FOnlineSubsystemGooglePlay::GetSessionInterface() const
+{
+	return nullptr;
+}
+
+
+IOnlineFriendsPtr FOnlineSubsystemGooglePlay::GetFriendsInterface() const
+{
+	return nullptr;
+}
+
+
+IOnlineSharedCloudPtr FOnlineSubsystemGooglePlay::GetSharedCloudInterface() const
+{
+	return nullptr;
+}
+
+
+IOnlineUserCloudPtr FOnlineSubsystemGooglePlay::GetUserCloudInterface() const
+{
+	return nullptr;
+}
+
+IOnlineLeaderboardsPtr FOnlineSubsystemGooglePlay::GetLeaderboardsInterface() const
+{
+	return nullptr;
+}
+
+
+IOnlineVoicePtr FOnlineSubsystemGooglePlay::GetVoiceInterface() const
+{
+	return nullptr;
+}
+
+
+IOnlineExternalUIPtr FOnlineSubsystemGooglePlay::GetExternalUIInterface() const
+{
+	return nullptr;
+}
+
+
+IOnlineTimePtr FOnlineSubsystemGooglePlay::GetTimeInterface() const
+{
+	return nullptr;
+}
+
+
+
+IOnlineTitleFilePtr FOnlineSubsystemGooglePlay::GetTitleFileInterface() const
+{
+	return nullptr;
+}
+
+
+IOnlineEntitlementsPtr FOnlineSubsystemGooglePlay::GetEntitlementsInterface() const
+{
+	return nullptr;
+}
+
+IOnlineAchievementsPtr FOnlineSubsystemGooglePlay::GetAchievementsInterface() const
+{
+	return AchievementsInterface;
+}
+
+bool FOnlineSubsystemGooglePlay::Init() 
+{
+	FPlatformMisc::LowLevelOutputDebugStringf(TEXT("FOnlineSubsystemAndroid::Init"));
+	
+	IdentityInterface = MakeShareable(new FOnlineIdentityGooglePlay(this));
+	AchievementsInterface = MakeShareable(new FOnlineAchievementsGooglePlay(this));
+
+	extern void AndroidThunkCpp_GooglePlayConnect();
+	AndroidThunkCpp_GooglePlayConnect();
+
+	return true;
+}
+
+bool FOnlineSubsystemGooglePlay::Tick(float DeltaTime)
+{
+	return true;
+}
+
+
+bool FOnlineSubsystemGooglePlay::Shutdown() 
+{
+	UE_LOG(LogOnline, Log, TEXT("FOnlineSubsystemAndroid::Shutdown()"));
+
+#define DESTRUCT_INTERFACE(Interface) \
+	if (Interface.IsValid()) \
+	{ \
+	ensure(Interface.IsUnique()); \
+	Interface = NULL; \
+	}
+
+	// Destruct the interfaces
+	DESTRUCT_INTERFACE(AchievementsInterface);
+	DESTRUCT_INTERFACE(IdentityInterface);
+#undef DESTRUCT_INTERFACE
+
+	return true;
+}
+
+
+FString FOnlineSubsystemGooglePlay::GetAppId() const 
+{
+	//get app id from settings. 
+	return TEXT( "AndroidAppIDPlaceHolder" );
+}
+
+
+bool FOnlineSubsystemGooglePlay::Exec(UWorld* InWorld, const TCHAR* Cmd, FOutputDevice& Ar) 
+{
+	return false;
+}
+
+bool FOnlineSubsystemGooglePlay::IsEnabled(void)
+{
+	bool bEnabled = true;
+	GConfig->GetBool(TEXT("/Script/AndroidRuntimeSettings.AndroidRuntimeSettings"), TEXT("bEnableGooglePlaySupport"), bEnabled, GEngineIni);
+	return bEnabled;
+}
diff --git a/Engine/Source/Runtime/Online/Android/OnlineSubsystemGooglePlay/Private/OnlineSubsystemGooglePlayModule.cpp b/Engine/Source/Runtime/Online/Android/OnlineSubsystemGooglePlay/Private/OnlineSubsystemGooglePlayModule.cpp
new file mode 100644
--- /dev/null
+++ b/Engine/Source/Runtime/Online/Android/OnlineSubsystemGooglePlay/Private/OnlineSubsystemGooglePlayModule.cpp
@@ -0,0 +1,83 @@
+// Copyright 1998-2014 Epic Games, Inc. All Rights Reserved.
+
+#include "OnlineSubsystemGooglePlayPrivatePCH.h"
+#include "OnlineSubsystemGooglePlay.h"
+
+IMPLEMENT_MODULE( FOnlineSubsystemGooglePlayModule, OnlineSubsystemGooglePlay );
+
+/**
+ * Class responsible for creating instance(s) of the subsystem
+ */
+class FOnlineFactoryGooglePlay : public IOnlineFactory
+{
+
+private:
+
+	/** Single instantiation of the IOS interface */
+	static FOnlineSubsystemGooglePlayPtr GooglePlaySingleton;
+
+	virtual void DestroySubsystem()
+	{
+		if (GooglePlaySingleton.IsValid())
+		{
+			GooglePlaySingleton->Shutdown();
+			GooglePlaySingleton = NULL;
+		}
+	}
+
+public:
+
+	FOnlineFactoryGooglePlay() {}
+	virtual ~FOnlineFactoryGooglePlay() 
+	{
+		DestroySubsystem();
+	}
+
+	virtual IOnlineSubsystemPtr CreateSubsystem(FName InstanceName)
+	{
+		if (!GooglePlaySingleton.IsValid())
+		{
+			GooglePlaySingleton = MakeShareable(new FOnlineSubsystemGooglePlay());
+			if (GooglePlaySingleton->IsEnabled())
+			{
+				if(!GooglePlaySingleton->Init())
+				{
+					UE_LOG(LogOnline, Warning, TEXT("FOnlineSubsystemGooglePlayModule failed to initialize!"));
+					DestroySubsystem();
+				}
+			}
+			else
+			{
+				UE_LOG(LogOnline, Warning, TEXT("FOnlineSubsystemGooglePlayModule was disabled"));
+				DestroySubsystem();
+			}
+
+			return GooglePlaySingleton;
+		}
+
+		UE_LOG_ONLINE(Warning, TEXT("Can't create more than one instance of a Google Play online subsystem!"));
+		return NULL;
+	}
+};
+
+FOnlineSubsystemGooglePlayPtr FOnlineFactoryGooglePlay::GooglePlaySingleton = NULL;
+
+
+void FOnlineSubsystemGooglePlayModule::StartupModule()
+{
+	UE_LOG(LogOnline, Display, TEXT("OnlineSubsystemGooglePlayModule::StartupModule()"));
+
+	GooglePlayFactory = new FOnlineFactoryGooglePlay();
+
+	FOnlineSubsystemModule& OSS = FModuleManager::GetModuleChecked<FOnlineSubsystemModule>("OnlineSubsystem");
+	OSS.RegisterPlatformService(GOOGLEPLAY_SUBSYSTEM, GooglePlayFactory);
+}
+
+
+void FOnlineSubsystemGooglePlayModule::ShutdownModule()
+{
+	UE_LOG(LogOnline, Display, TEXT("FOnlineSubsystemGooglePlayModule::ShutdownModule()"));
+
+	delete GooglePlayFactory;
+	GooglePlayFactory = nullptr;
+}
diff --git a/Engine/Source/Runtime/Online/Android/OnlineSubsystemGooglePlay/Private/OnlineSubsystemGooglePlayPrivatePCH.h b/Engine/Source/Runtime/Online/Android/OnlineSubsystemGooglePlay/Private/OnlineSubsystemGooglePlayPrivatePCH.h
new file mode 100644
--- /dev/null
+++ b/Engine/Source/Runtime/Online/Android/OnlineSubsystemGooglePlay/Private/OnlineSubsystemGooglePlayPrivatePCH.h
@@ -0,0 +1,16 @@
+// Copyright 1998-2014 Epic Games, Inc. All Rights Reserved.
+
+#pragma once
+
+/** FName declaration of Steam subsystem */
+#define GOOGLEPLAY_SUBSYSTEM FName(TEXT("GooglePlay"))
+
+#include "Core.h"
+#include "SocketSubsystem.h"
+#include "ModuleManager.h"
+
+#include "OnlineAchievementsInterfaceGooglePlay.h"
+
+#include "OnlineSubsystemGooglePlay.h"
+#include "OnlineSubsystemGooglePlayModule.h"
+
diff --git a/Engine/Source/Runtime/Online/Android/OnlineSubsystemGooglePlay/Public/OnlineSubsystemGooglePlay.h b/Engine/Source/Runtime/Online/Android/OnlineSubsystemGooglePlay/Public/OnlineSubsystemGooglePlay.h
new file mode 100644
--- /dev/null
+++ b/Engine/Source/Runtime/Online/Android/OnlineSubsystemGooglePlay/Public/OnlineSubsystemGooglePlay.h
@@ -0,0 +1,74 @@
+// Copyright 1998-2014 Epic Games, Inc. All Rights Reserved.
+
+#pragma once
+
+#include "OnlineSubsystem.h"
+#include "OnlineIdentityInterfaceGooglePlay.h"
+#include "OnlineAchievementsInterfaceGooglePlay.h"
+
+/**
+ * OnlineSubsystemGooglePlay - Implementation of the online subsystem for Google Play services
+ */
+class ONLINESUBSYSTEMGOOGLEPLAY_API FOnlineSubsystemGooglePlay : 
+	public IOnlineSubsystem,
+	public FTickerObjectBase
+{
+public:
+	FOnlineSubsystemGooglePlay();
+	virtual ~FOnlineSubsystemGooglePlay() {}
+
+	// Begin IOnlineSubsystem Interface
+	virtual IOnlineSessionPtr GetSessionInterface() const OVERRIDE;
+	virtual IOnlineFriendsPtr GetFriendsInterface() const OVERRIDE;
+	virtual IOnlineSharedCloudPtr GetSharedCloudInterface() const OVERRIDE;
+	virtual IOnlineUserCloudPtr GetUserCloudInterface() const OVERRIDE;
+	virtual IOnlineLeaderboardsPtr GetLeaderboardsInterface() const OVERRIDE;
+	virtual IOnlineVoicePtr GetVoiceInterface() const  OVERRIDE;
+	virtual IOnlineExternalUIPtr GetExternalUIInterface() const OVERRIDE;
+	virtual IOnlineTimePtr GetTimeInterface() const OVERRIDE;
+	virtual IOnlineIdentityPtr GetIdentityInterface() const OVERRIDE;
+	virtual IOnlineTitleFilePtr GetTitleFileInterface() const OVERRIDE;
+	virtual IOnlineEntitlementsPtr GetEntitlementsInterface() const OVERRIDE;
+	virtual IOnlineStorePtr GetStoreInterface() const OVERRIDE;
+	virtual IOnlineEventsPtr GetEventsInterface() const OVERRIDE { return NULL; }
+	virtual IOnlineMessagePtr GetMessageInterface() const OVERRIDE { return NULL; }
+	virtual IOnlineSharingPtr GetSharingInterface() const OVERRIDE { return NULL; }
+	virtual IOnlineUserPtr GetUserInterface() const OVERRIDE { return NULL; }
+	virtual IOnlineAchievementsPtr GetAchievementsInterface() const OVERRIDE;
+	virtual IOnlinePresencePtr GetPresenceInterface() const OVERRIDE { return NULL; }
+	virtual class UObject* GetNamedInterface(FName InterfaceName) OVERRIDE { return NULL; }
+	virtual void SetNamedInterface(FName InterfaceName, class UObject* NewInterface) OVERRIDE {}
+	virtual bool IsDedicated() const OVERRIDE { return false; }
+	virtual bool IsServer() const OVERRIDE { return true; }
+	virtual void SetForceDedicated(bool bForce) OVERRIDE {}
+	virtual bool IsLocalPlayer(const FUniqueNetId& UniqueId) const OVERRIDE { return true; }
+
+	virtual bool Init() OVERRIDE;
+	virtual bool Shutdown() OVERRIDE;
+	virtual FString GetAppId() const OVERRIDE;
+	virtual bool Exec(class UWorld* InWorld, const TCHAR* Cmd, FOutputDevice& Ar) OVERRIDE;
+	// End IOnlineSubsystem Interface
+
+	virtual bool Tick(float DeltaTime) OVERRIDE;
+	
+PACKAGE_SCOPE:
+
+	/**
+	 * Is Online Subsystem Android available for use
+	 * @return true if Android Online Subsystem functionality is available, false otherwise
+	 */
+	bool IsEnabled();
+
+private:
+
+	/** Online async task thread */
+	class FRunnableThread* OnlineAsyncTaskThread;
+
+	/** Interface to the online identity system */
+	FOnlineIdentityGooglePlayPtr IdentityInterface;
+
+	/** Interface to the online achievements */
+	FOnlineAchievementsGooglePlayPtr AchievementsInterface;
+};
+
+typedef TSharedPtr<FOnlineSubsystemGooglePlay, ESPMode::ThreadSafe> FOnlineSubsystemGooglePlayPtr;
diff --git a/Engine/Source/Runtime/Online/Android/OnlineSubsystemGooglePlay/Public/OnlineSubsystemGooglePlayModule.h b/Engine/Source/Runtime/Online/Android/OnlineSubsystemGooglePlay/Public/OnlineSubsystemGooglePlayModule.h
new file mode 100644
--- /dev/null
+++ b/Engine/Source/Runtime/Online/Android/OnlineSubsystemGooglePlay/Public/OnlineSubsystemGooglePlayModule.h
@@ -0,0 +1,38 @@
+// Copyright 1998-2014 Epic Games, Inc. All Rights Reserved.
+
+#pragma once
+
+#include "Core.h"
+#include "ModuleInterface.h"
+
+/**
+ * Online subsystem module class (Google Play Implementation)
+ * Code related to the loading and handling of the Amdroid Google Play module.
+ */
+class FOnlineSubsystemGooglePlayModule : public IModuleInterface
+{
+private:
+	
+	/** Class responsible for creating instance(s) of the subsystem */
+	class FOnlineFactoryGooglePlay* GooglePlayFactory;
+
+	virtual ~FOnlineSubsystemGooglePlayModule(){}
+
+	// Begin IModuleInterface interface
+	virtual void StartupModule() OVERRIDE;
+	virtual void ShutdownModule() OVERRIDE;
+
+	virtual bool SupportsDynamicReloading() OVERRIDE
+	{
+		return false;
+	}
+
+	virtual bool SupportsAutomaticShutdown() OVERRIDE
+	{
+		return false;
+	}
+	// End IModuleInterface interface
+};
+
+typedef TSharedPtr<FOnlineSubsystemGooglePlayModule> FOnlineSubsystemGooglePlayModulePtr;
+
diff --git a/Engine/Source/Runtime/Online/Android/OnlineSubsystemGooglePlay/Public/OnlineSubsystemGooglePlayPackage.h b/Engine/Source/Runtime/Online/Android/OnlineSubsystemGooglePlay/Public/OnlineSubsystemGooglePlayPackage.h
new file mode 100644
--- /dev/null
+++ b/Engine/Source/Runtime/Online/Android/OnlineSubsystemGooglePlay/Public/OnlineSubsystemGooglePlayPackage.h
@@ -0,0 +1,13 @@
+// Copyright 1998-2014 Epic Games, Inc. All Rights Reserved.
+
+// Can't be #pragma once because other modules may define PACKAGE_SCOPE
+
+// Intended to be the last include in an exported class definition
+// Properly defines some members as "public to the module" vs "private to the consumer/user"
+
+#undef PACKAGE_SCOPE
+#ifdef ONLINESUBSYSTEMGOOGLEPLAY_PACKAGE
+#define PACKAGE_SCOPE public
+#else
+#define PACKAGE_SCOPE protected
+#endif
diff --git a/Engine/Source/Runtime/Online/OnlineSubsystem/Public/Interfaces/OnlineExternalUIInterface.h b/Engine/Source/Runtime/Online/OnlineSubsystem/Public/Interfaces/OnlineExternalUIInterface.h
--- a/Engine/Source/Runtime/Online/OnlineSubsystem/Public/Interfaces/OnlineExternalUIInterface.h
+++ b/Engine/Source/Runtime/Online/OnlineSubsystem/Public/Interfaces/OnlineExternalUIInterface.h
@@ -1,6 +1,7 @@
 // Copyright 1998-2014 Epic Games, Inc. All Rights Reserved.
 
 #pragma once
+#include "OnlineSubsystemTypes.h"
 #include "OnlineDelegateMacros.h"
 
 /**
diff --git a/Engine/Source/Runtime/Online/OnlineSubsystem/Public/OnlineStats.h b/Engine/Source/Runtime/Online/OnlineSubsystem/Public/OnlineStats.h
--- a/Engine/Source/Runtime/Online/OnlineSubsystem/Public/OnlineStats.h
+++ b/Engine/Source/Runtime/Online/OnlineSubsystem/Public/OnlineStats.h
@@ -3,6 +3,7 @@
 #pragma once
 
 #include "OnlineKeyValuePair.h"
+#include "OnlineSubsystemTypes.h"
 
 /** Representation of a single column and its data */ 
 typedef FOnlineKeyValuePairs<FName, FVariantData> FStatsColumnArray;
diff --git a/Engine/Source/Runtime/UE4Game/UE4Game.Build.cs b/Engine/Source/Runtime/UE4Game/UE4Game.Build.cs
--- a/Engine/Source/Runtime/UE4Game/UE4Game.Build.cs
+++ b/Engine/Source/Runtime/UE4Game/UE4Game.Build.cs
@@ -17,9 +17,9 @@ public class UE4Game : ModuleRules
 		}
 		else if (Target.Platform == UnrealTargetPlatform.Android)
 		{
-			// @todo android: Add Android online subsystem
 			PrivateDependencyModuleNames.AddRange(new string[] { "Core", "CoreUObject", "Engine", "OnlineSubsystem", "OnlineSubsystemUtils" });
 			DynamicallyLoadedModuleNames.Add("AndroidAdvertising");
+			DynamicallyLoadedModuleNames.Add("OnlineSubsystemGooglePlay");
 		}
 	}
 }
