Added AdMob support for Android.

- In Edit/Project Settings/Android, set an AdMob Ad ID and then show and hide ads via blueprints.

diff --git a/Engine/Build/Android/Java/AndroidManifest.xml b/Engine/Build/Android/Java/AndroidManifest.xml
--- a/Engine/Build/Android/Java/AndroidManifest.xml
+++ b/Engine/Build/Android/Java/AndroidManifest.xml
@@ -25,6 +25,8 @@
 			android:value="@string/app_id" />
 		<meta-data android:name="com.google.android.gms.version"
            android:value="@integer/google_play_services_version" />
+		<activity android:name="com.google.android.gms.ads.AdActivity"
+				  android:configChanges="keyboard|keyboardHidden|orientation|screenLayout|uiMode|screenSize|smallestScreenSize"/>
 	</application>
 	
 	<!-- Requirements -->
diff --git a/Engine/Build/Android/Java/src/com/epicgames/ue4/GameActivity.java b/Engine/Build/Android/Java/src/com/epicgames/ue4/GameActivity.java
--- a/Engine/Build/Android/Java/src/com/epicgames/ue4/GameActivity.java
+++ b/Engine/Build/Android/Java/src/com/epicgames/ue4/GameActivity.java
@@ -17,6 +17,13 @@ import android.content.Intent;
 import android.content.res.Configuration;
 import android.content.IntentSender.SendIntentException;
 
+import android.view.Gravity;
+import android.view.ViewGroup.LayoutParams;
+import android.view.ViewGroup.MarginLayoutParams;
+import android.view.WindowManager;
+import android.widget.LinearLayout;
+import android.widget.PopupWindow;
+
 import com.google.android.gms.common.api.GoogleApiClient;
 import com.google.android.gms.common.api.PendingResult;
 import com.google.android.gms.common.api.ResultCallback;
@@ -25,7 +32,9 @@ import com.google.android.gms.common.GooglePlayServicesUtil;
 import com.google.android.gms.games.achievement.*;
 import com.google.android.gms.games.Games;
 
-
+import com.google.android.gms.ads.AdRequest;
+import com.google.android.gms.ads.AdView;
+import com.google.android.gms.ads.AdSize;
 
 // TODO: use the resources from the UE4 lib project once we've got the packager up and running
 //import com.epicgames.ue4.R;
@@ -52,6 +61,13 @@ public class GameActivity extends NativeActivity implements GoogleApiClient.Conn
 	private boolean bResolvingGoogleServicesError = false;
 	private int dialogError = 0;
 
+	/** AdMob support */
+	private PopupWindow adPopupWindow;
+	private AdView adView;
+	private boolean adInit = false;
+	private LinearLayout adLayout;
+	private LinearLayout activityLayout;
+
 	/** Request code to use when launching the Google Services resolution activity */
     private static final int GOOGLE_SERVICES_REQUEST_RESOLVE_ERROR = 1001;
 
@@ -158,11 +174,11 @@ public class GameActivity extends NativeActivity implements GoogleApiClient.Conn
 
 		// Connect to Google Play Services
 		googleClient = new GoogleApiClient.Builder(this)
-         .addApi(Games.API)
-         .addScope(Games.SCOPE_GAMES)
-         .addConnectionCallbacks(this)
-         .addOnConnectionFailedListener(this)
-         .build();
+		 .addApi(Games.API)
+		 .addScope(Games.SCOPE_GAMES)
+		 .addConnectionCallbacks(this)
+		 .addOnConnectionFailedListener(this)
+		 .build();
 	}
 
 	@Override
@@ -205,7 +221,7 @@ public class GameActivity extends NativeActivity implements GoogleApiClient.Conn
 	}
 
 	// Callbacks to handle connections with Google Play
-	@Override
+	 @Override
     public void onConnected(Bundle connectionHint)
 	{
         Log.debug("Connected to Google Play Services.");
@@ -328,8 +344,8 @@ public class GameActivity extends NativeActivity implements GoogleApiClient.Conn
 
 	public void AndroidThunkJava_GooglePlayConnect()
 	{
-		googleClient.connect();
-	}
+			googleClient.connect();
+		}
 
 	public void AndroidThunkJava_ShowLeaderboard(String LeaderboardID)
 	{
@@ -413,6 +429,131 @@ public class GameActivity extends NativeActivity implements GoogleApiClient.Conn
 		}
 	}
 
+	public void AndroidThunkJava_ShowAdBanner(String AdMobAdUnitID, boolean bShowOnBottonOfScreen)
+	{
+		Log.debug("In AndroidThunkJava_ShowAdBanner");
+		Log.debug("AdID: " + AdMobAdUnitID);
+
+		final int adGravity = bShowOnBottonOfScreen ? Gravity.BOTTOM : Gravity.TOP;
+
+		if (adInit)
+		{
+			// already created, make it visible
+			_activity.runOnUiThread(new Runnable()
+			{
+				@Override
+				public void run()
+				{
+					if ((adPopupWindow == null) || adPopupWindow.isShowing())
+					{
+						return;
+					}
+
+					adPopupWindow.showAtLocation(activityLayout, adGravity, 0, 0);
+					adPopupWindow.update();
+				}
+			});
+
+			return;
+		}
+
+		// init our AdMob window
+		adView = new AdView(this);
+		adView.setAdUnitId(AdMobAdUnitID);
+		adView.setAdSize(AdSize.BANNER);
+
+		if (adView != null)
+		{
+			_activity.runOnUiThread(new Runnable()
+			{
+				@Override
+				public void run()
+				{
+					adInit = true;
+
+					adPopupWindow = new PopupWindow(_activity);
+					adPopupWindow.setWidth(320);
+					adPopupWindow.setHeight(50);
+					adPopupWindow.setWindowLayoutMode(LayoutParams.WRAP_CONTENT, LayoutParams.WRAP_CONTENT);
+					adPopupWindow.setClippingEnabled(false);
+
+					adLayout = new LinearLayout(_activity);
+					activityLayout = new LinearLayout(_activity);
+
+					adLayout.setPadding(-5,-5,-5,-5);
+
+					MarginLayoutParams params = new MarginLayoutParams(LayoutParams.WRAP_CONTENT, LayoutParams.WRAP_CONTENT);;
+
+					params.setMargins(0,0,0,0);
+
+					adLayout.setOrientation(LinearLayout.VERTICAL);
+					adLayout.addView(adView, params);
+					adPopupWindow.setContentView(adLayout);
+
+					_activity.setContentView(activityLayout, params);
+
+					AdRequest adRequest = new AdRequest.Builder().build();		// add test devices here
+
+					_activity.adView.loadAd(adRequest);
+
+					adPopupWindow.showAtLocation(activityLayout, adGravity, 0, 0);
+					adPopupWindow.update();
+				}
+			});
+		}
+	}
+
+	public void AndroidThunkJava_HideAdBanner()
+	{
+		Log.debug("In AndroidThunkJava_HideAdBanner");
+
+		if (!adInit)
+		{
+			return;
+		}
+
+		_activity.runOnUiThread(new Runnable()
+		{
+			@Override
+			public void run()
+			{
+				if ((adPopupWindow == null) || !adPopupWindow.isShowing())
+				{
+					return;
+				}
+
+				adPopupWindow.dismiss();
+				adPopupWindow.update();
+			}
+		});
+	}
+
+	public void AndroidThunkJava_CloseAdBanner()
+	{
+		Log.debug("In AndroidThunkJava_CloseAdBanner");
+
+		if (!adInit)
+		{
+			return;
+		}
+
+		// currently the same as hide.  should we do a full teardown?
+		_activity.runOnUiThread(new Runnable()
+		{
+			@Override
+			public void run()
+			{
+				if ((adPopupWindow == null) || !adPopupWindow.isShowing())
+				{
+					return;
+				}
+
+				adPopupWindow.dismiss();
+				adPopupWindow.update();
+			}
+		});
+	}
+
 	public native boolean nativeIsShippingBuild();
 	public native void nativeSetGlobalActivity();
 	public native void nativeSetWindowInfo(boolean bIsPortrait);
diff --git a/Engine/Source/Developer/Android/AndroidPlatformEditor/Private/AndroidTargetSettingsCustomization.cpp b/Engine/Source/Developer/Android/AndroidPlatformEditor/Private/AndroidTargetSettingsCustomization.cpp
--- a/Engine/Source/Developer/Android/AndroidPlatformEditor/Private/AndroidTargetSettingsCustomization.cpp
+++ b/Engine/Source/Developer/Android/AndroidPlatformEditor/Private/AndroidTargetSettingsCustomization.cpp
@@ -151,6 +151,10 @@ void FAndroidTargetSettingsCustomization::BuildAppManifestSection(IDetailLayoutB
 	AppIDProperty->SetOnPropertyValueChanged(FSimpleDelegate::CreateRaw(this, &FAndroidTargetSettingsCustomization::OnAppIDModified));
 	GooglePlayCategory.AddProperty(AppIDProperty)
 		.EditCondition(SetupForGooglePlayAttribute, NULL);
+
+	TSharedRef<IPropertyHandle> AdMobAdUnitIDProperty = DetailLayout.GetProperty(GET_MEMBER_NAME_CHECKED(UAndroidRuntimeSettings, AdMobAdUnitID));
+	GooglePlayCategory.AddProperty(AdMobAdUnitIDProperty)
+		.EditCondition(SetupForGooglePlayAttribute, NULL);
 }
 
 void FAndroidTargetSettingsCustomization::BuildIconSection(IDetailLayoutBuilder& DetailLayout)
diff --git a/Engine/Source/Runtime/Advertising/Android/AndroidAdvertising/Private/AndroidAdvertising.cpp b/Engine/Source/Runtime/Advertising/Android/AndroidAdvertising/Private/AndroidAdvertising.cpp
--- a/Engine/Source/Runtime/Advertising/Android/AndroidAdvertising/Private/AndroidAdvertising.cpp
+++ b/Engine/Source/Runtime/Advertising/Android/AndroidAdvertising/Private/AndroidAdvertising.cpp
@@ -10,12 +10,24 @@ IMPLEMENT_MODULE( FAndroidAdvertisingProvider, AndroidAdvertising );
 
 void FAndroidAdvertisingProvider::ShowAdBanner( bool bShowOnBottomOfScreen ) 
 {
+	FString AdUnitID;
+	bool found = GConfig->GetString(TEXT("/Script/AndroidRuntimeSettings.AndroidRuntimeSettings"), TEXT("AdMobAdUnitID"), AdUnitID, GEngineIni);
+
+	if (!found || AdUnitID.IsEmpty())
+		return;
+
+	extern void AndroidThunkCpp_ShowAdBanner(const FString&, bool);
+	AndroidThunkCpp_ShowAdBanner(AdUnitID, bShowOnBottomOfScreen);
 }
 
 void FAndroidAdvertisingProvider::HideAdBanner() 
 {
+	extern void AndroidThunkCpp_HideAdBanner();
+	AndroidThunkCpp_HideAdBanner();
 }
 
 void FAndroidAdvertisingProvider::CloseAdBanner() 
 {
+	extern void AndroidThunkCpp_CloseAdBanner();
+	AndroidThunkCpp_CloseAdBanner();
 }
diff --git a/Engine/Source/Runtime/Android/AndroidRuntimeSettings/Classes/AndroidRuntimeSettings.h b/Engine/Source/Runtime/Android/AndroidRuntimeSettings/Classes/AndroidRuntimeSettings.h
--- a/Engine/Source/Runtime/Android/AndroidRuntimeSettings/Classes/AndroidRuntimeSettings.h
+++ b/Engine/Source/Runtime/Android/AndroidRuntimeSettings/Classes/AndroidRuntimeSettings.h
@@ -97,4 +97,8 @@ public:
 	// Mapping of game leaderboard names to IDs generated by Google Play.
 	UPROPERTY(GlobalConfig, EditAnywhere, Category = GooglePlayServices)
 	TArray<FGooglePlayLeaderboardMapping> LeaderboardMap;
+
+	// The unique identifier for the ad obtained from AdMob.
+	UPROPERTY(GlobalConfig, EditAnywhere, Category = GooglePlayServices)
+	FString AdMobAdUnitID;
 };
diff --git a/Engine/Source/Runtime/Launch/Private/Android/AndroidJNI.cpp b/Engine/Source/Runtime/Launch/Private/Android/AndroidJNI.cpp
--- a/Engine/Source/Runtime/Launch/Private/Android/AndroidJNI.cpp
+++ b/Engine/Source/Runtime/Launch/Private/Android/AndroidJNI.cpp
@@ -99,6 +99,9 @@ jmethodID JDef_GameActivity::AndroidThunkJava_ShowAchievements;
 jmethodID JDef_GameActivity::AndroidThunkJava_WriteLeaderboardValue;
 jmethodID JDef_GameActivity::AndroidThunkJava_GooglePlayConnect;
 jmethodID JDef_GameActivity::AndroidThunkJava_WriteAchievement;
+jmethodID JDef_GameActivity::AndroidThunkJava_ShowAdBanner;
+jmethodID JDef_GameActivity::AndroidThunkJava_HideAdBanner;
+jmethodID JDef_GameActivity::AndroidThunkJava_CloseAdBanner;
 
 DEFINE_LOG_CATEGORY_STATIC(LogEngine, Log, All);
 
@@ -212,6 +215,33 @@ void AndroidThunkCpp_WriteAchievement(const FString& AchievementID, float Percen
 	}
 }
 
+void AndroidThunkCpp_ShowAdBanner(const FString& AdUnitID, bool bShowOnBottomOfScreen)
+{
+	if (JNIEnv* Env = GetJavaEnv())
+ 	{
+		jstring AdUnitIDArg = Env->NewStringUTF(TCHAR_TO_UTF8(*AdUnitID));
+		Env->CallVoidMethod(GJavaGlobalThis, JDef_GameActivity::AndroidThunkJava_ShowAdBanner, AdUnitIDArg, bShowOnBottomOfScreen);
+		Env->DeleteLocalRef(AdUnitIDArg);
+ 	}
+}
+
+void AndroidThunkCpp_HideAdBanner()
+{
+ 	if (JNIEnv* Env = GetJavaEnv())
+ 	{
+		Env->CallVoidMethod(GJavaGlobalThis, JDef_GameActivity::AndroidThunkJava_HideAdBanner);
+ 	}
+}
+
+void AndroidThunkCpp_CloseAdBanner()
+{
+ 	if (JNIEnv* Env = GetJavaEnv())
+ 	{
+		Env->CallVoidMethod(GJavaGlobalThis, JDef_GameActivity::AndroidThunkJava_CloseAdBanner);
+ 	}
+}
+
+
 //The JNI_OnLoad function is triggered by loading the game library from 
 //the Java source file.
 //	static
@@ -240,6 +270,9 @@ JNIEXPORT jint JNI_OnLoad(JavaVM* InJavaVM, void* InReserved)
 	JDef_GameActivity::AndroidThunkJava_WriteLeaderboardValue = env->GetMethodID(JDef_GameActivity::ClassID, "AndroidThunkJava_WriteLeaderboardValue", "(Ljava/lang/String;J)V");
 	JDef_GameActivity::AndroidThunkJava_GooglePlayConnect = env->GetMethodID(JDef_GameActivity::ClassID, "AndroidThunkJava_GooglePlayConnect", "()V");
 	JDef_GameActivity::AndroidThunkJava_WriteAchievement = env->GetMethodID(JDef_GameActivity::ClassID, "AndroidThunkJava_WriteAchievement", "(Ljava/lang/String;F)V");
+	JDef_GameActivity::AndroidThunkJava_ShowAdBanner = env->GetMethodID(JDef_GameActivity::ClassID, "AndroidThunkJava_ShowAdBanner", "(Ljava/lang/String;Z)V");
+	JDef_GameActivity::AndroidThunkJava_HideAdBanner = env->GetMethodID(JDef_GameActivity::ClassID, "AndroidThunkJava_HideAdBanner", "()V");
+	JDef_GameActivity::AndroidThunkJava_CloseAdBanner = env->GetMethodID(JDef_GameActivity::ClassID, "AndroidThunkJava_CloseAdBanner", "()V");
 
 	// hook signals
 #if UE_BUILD_DEBUG
diff --git a/Engine/Source/Runtime/Launch/Private/Android/AndroidJNI.h b/Engine/Source/Runtime/Launch/Private/Android/AndroidJNI.h
--- a/Engine/Source/Runtime/Launch/Private/Android/AndroidJNI.h
+++ b/Engine/Source/Runtime/Launch/Private/Android/AndroidJNI.h
@@ -26,6 +26,9 @@ public:
 	static jmethodID AndroidThunkJava_WriteLeaderboardValue;
 	static jmethodID AndroidThunkJava_GooglePlayConnect;
 	static jmethodID AndroidThunkJava_WriteAchievement;
+	static jmethodID AndroidThunkJava_ShowAdBanner;
+	static jmethodID AndroidThunkJava_HideAdBanner;
+	static jmethodID AndroidThunkJava_CloseAdBanner;
 };
 
 // Returns the java environment
