Android, favor the volume keys adjusting the media volume as that is the game volume.

diff --git a/Engine/Build/Android/Java/src/com/epicgames/ue4/GameActivity.java b/Engine/Build/Android/Java/src/com/epicgames/ue4/GameActivity.java
--- a/Engine/Build/Android/Java/src/com/epicgames/ue4/GameActivity.java
+++ b/Engine/Build/Android/Java/src/com/epicgames/ue4/GameActivity.java
@@ -17,6 +17,8 @@ import android.content.Intent;
 import android.content.res.Configuration;
 import android.content.IntentSender.SendIntentException;
 
+import android.media.AudioManager;
+
 import android.view.Gravity;
 import android.view.ViewGroup.LayoutParams;
 import android.view.ViewGroup.MarginLayoutParams;
@@ -129,6 +131,9 @@ public class GameActivity extends NativeActivity implements GoogleApiClient.Conn
 
 		_activity = this;
 
+		// tell Android that we want volume controls to change the media volume, aka music
+		setVolumeControlStream(AudioManager.STREAM_MUSIC);
+
 		// tell the engine if this is a portrait app
 		nativeSetGlobalActivity();
 		nativeSetWindowInfo(getResources().getConfiguration().orientation == Configuration.ORIENTATION_PORTRAIT);
