Merging in Google Play connection fix

If the game has successfully connected to Google play, try to reconnect in the activity's onStart function.

diff --git a/Engine/Build/Android/Java/src/com/epicgames/ue4/GameActivity.java b/Engine/Build/Android/Java/src/com/epicgames/ue4/GameActivity.java
--- a/Engine/Build/Android/Java/src/com/epicgames/ue4/GameActivity.java
+++ b/Engine/Build/Android/Java/src/com/epicgames/ue4/GameActivity.java
@@ -63,6 +63,9 @@ public class GameActivity extends NativeActivity implements GoogleApiClient.Conn
 	private boolean bResolvingGoogleServicesError = false;
 	private int dialogError = 0;
 
+	/** Flag indicating that we successfully connected to Google Play. */
+	private boolean bHaveConnectedToGooglePlay = false;
+
 	/** AdMob support */
 	private PopupWindow adPopupWindow;
 	private AdView adView;
@@ -116,6 +119,12 @@ public class GameActivity extends NativeActivity implements GoogleApiClient.Conn
 		super.onStart();
 		
 		Log.debug("==================================> Inside onStart function in GameActivity");
+
+		// Reconnect to Google Play if we were connected before
+		if(bHaveConnectedToGooglePlay)
+		{
+			googleClient.connect();
+		}
 	}
 
 	@Override
@@ -231,6 +240,9 @@ public class GameActivity extends NativeActivity implements GoogleApiClient.Conn
 	{
         Log.debug("Connected to Google Play Services.");
 
+		// Set the flag that we successfully connected. Checked in onStart to re-establish the connection.
+		bHaveConnectedToGooglePlay = true;
+
 		// Load achievements. Since games are expected to pass in achievement progress as a percentage,
 		// we need to know what the maximum steps are in order to convert the percentage to an integer
 		// number of steps.
