Fixed up tappychicken on the shield starting in portrait.

diff --git a/Engine/Build/Android/Java/src/com/epicgames/ue4/GameActivity.java b/Engine/Build/Android/Java/src/com/epicgames/ue4/GameActivity.java
--- a/Engine/Build/Android/Java/src/com/epicgames/ue4/GameActivity.java
+++ b/Engine/Build/Android/Java/src/com/epicgames/ue4/GameActivity.java
@@ -127,6 +127,29 @@ public class GameActivity extends NativeActivity implements GoogleApiClient.Conn
 		}
 	}
 
+	public int getDeviceDefaultOrientation() 
+	{
+
+		// WindowManager windowManager =  (WindowManager) getSystemService(WINDOW_SERVICE);
+		WindowManager windowManager =  getWindowManager();
+
+		Configuration config = getResources().getConfiguration();
+
+		int rotation = windowManager.getDefaultDisplay().getRotation();
+
+		if ( ((rotation == android.view.Surface.ROTATION_0 || rotation == android.view.Surface.ROTATION_180) &&
+				config.orientation == Configuration.ORIENTATION_LANDSCAPE)
+			|| ((rotation == android.view.Surface.ROTATION_90 || rotation == android.view.Surface.ROTATION_270) &&    
+				config.orientation == Configuration.ORIENTATION_PORTRAIT)) 
+		{
+			return Configuration.ORIENTATION_LANDSCAPE;
+		}
+		else 
+		{
+			return Configuration.ORIENTATION_PORTRAIT;
+		}
+	}
+
 	@Override
 	public void onCreate(Bundle savedInstanceState)
 	{
@@ -142,7 +165,13 @@ public class GameActivity extends NativeActivity implements GoogleApiClient.Conn
 
 		// tell Android that we want volume controls to change the media volume, aka music
 		setVolumeControlStream(AudioManager.STREAM_MUSIC);
-
+		
+		if ( getDeviceDefaultOrientation() == Configuration.ORIENTATION_LANDSCAPE )
+		{
+			Log.debug( "Setting screen orientation to landscape because it's a Sheild" );
+			_activity.setRequestedOrientation( android.content.pm.ActivityInfo.SCREEN_ORIENTATION_LANDSCAPE );
+		}
+		
 		// tell the engine if this is a portrait app
 		nativeSetGlobalActivity();
 		nativeSetWindowInfo(getResources().getConfiguration().orientation == Configuration.ORIENTATION_PORTRAIT);
