Changed warning message for android landscape device detection so it was more correct.

diff --git a/Engine/Build/Android/Java/src/com/epicgames/ue4/GameActivity.java b/Engine/Build/Android/Java/src/com/epicgames/ue4/GameActivity.java
--- a/Engine/Build/Android/Java/src/com/epicgames/ue4/GameActivity.java
+++ b/Engine/Build/Android/Java/src/com/epicgames/ue4/GameActivity.java
@@ -172,7 +172,7 @@ public class GameActivity extends NativeActivity implements GoogleApiClient.Conn
 		
 		if ( getDeviceDefaultOrientation() == Configuration.ORIENTATION_LANDSCAPE )
 		{
-			Log.debug( "Setting screen orientation to landscape because it's a Sheild" );
+			Log.debug( "Setting screen orientation to landscape because we have detected landscape device" );
 			_activity.setRequestedOrientation( android.content.pm.ActivityInfo.SCREEN_ORIENTATION_LANDSCAPE );
 		}
 		
