Merging

//depot/UE4-TappyChicken/Engine/Build/...

to //depot/UE4/Engine/Build/...

diff --git a/Engine/Build/Android/Java/src/com/epicgames/ue4/GameActivity.java b/Engine/Build/Android/Java/src/com/epicgames/ue4/GameActivity.java
--- a/Engine/Build/Android/Java/src/com/epicgames/ue4/GameActivity.java
+++ b/Engine/Build/Android/Java/src/com/epicgames/ue4/GameActivity.java
@@ -19,6 +19,7 @@ import android.content.res.Configuration;
 import android.content.IntentSender.SendIntentException;
 
 import android.media.AudioManager;
+import android.util.DisplayMetrics;
 
 import android.view.Gravity;
 import android.view.ViewGroup.LayoutParams;
@@ -38,6 +39,7 @@ import com.google.android.gms.games.Games;
 import com.google.android.gms.ads.AdRequest;
 import com.google.android.gms.ads.AdView;
 import com.google.android.gms.ads.AdSize;
+import com.google.android.gms.ads.AdListener;
 
 // TODO: use the resources from the UE4 lib project once we've got the packager up and running
 //import com.epicgames.ue4.R;
@@ -76,6 +78,16 @@ public class GameActivity extends NativeActivity implements GoogleApiClient.Conn
 	private boolean adInit = false;
 	private LinearLayout adLayout;
 	private LinearLayout activityLayout;
+	private int adGravity = Gravity.TOP;
+
+	/** true when the application has requested that an ad be displayed */
+	private boolean adWantsToBeShown = false;
+
+	/** true when an ad is available to be displayed */
+	private boolean adIsAvailable = false;
+
+	/** true when an ad request is in flight */
+	private boolean adIsRequested = false;
 
 	/** Request code to use when launching the Google Services resolution activity */
     private static final int GOOGLE_SERVICES_REQUEST_RESOLVE_ERROR = 1001;
@@ -183,6 +195,13 @@ public class GameActivity extends NativeActivity implements GoogleApiClient.Conn
 		nativeSetGlobalActivity();
 		nativeSetWindowInfo(getResources().getConfiguration().orientation == Configuration.ORIENTATION_PORTRAIT);
 
+
+		Log.debug( "Android version is " + android.os.Build.VERSION.RELEASE );
+		Log.debug( "Android manufactorer is " + android.os.Build.MANUFACTURER );
+		Log.debug( "Android model is " + android.os.Build.MODEL );
+
+		nativeSetAndroidVersionInformation( android.os.Build.VERSION.RELEASE, android.os.Build.MANUFACTURER, android.os.Build.MODEL );
+
 		try
 		{
 			int Version = getPackageManager().getPackageInfo(getPackageName(), 0).versionCode;
@@ -195,6 +214,9 @@ public class GameActivity extends NativeActivity implements GoogleApiClient.Conn
 			Log.debug("==================================> PackageInfo failure getting .obb info: " + e.getMessage());
 		}
 		
+		// enable the physical volume controls to the game
+		this.setVolumeControlStream(AudioManager.STREAM_MUSIC);
+
 		final AlertDialog.Builder alertBuilder = new AlertDialog.Builder(this);
 		final EditText consoleInputBox = new EditText(this);
 		
@@ -409,6 +431,45 @@ public class GameActivity extends NativeActivity implements GoogleApiClient.Conn
 		}
 	}
 
+	// handle ad popup visibility and requests
+	private void updateAdVisibility(boolean loadIfNeeded)
+	{
+		if (!adInit || (adPopupWindow == null))
+		{
+			return;
+		}
+
+		// request an ad if we don't have one available or requested, but would like one
+		if (adWantsToBeShown && !adIsAvailable && !adIsRequested && loadIfNeeded)
+		{
+			AdRequest adRequest = new AdRequest.Builder().build();		// add test devices here
+			_activity.adView.loadAd(adRequest);
+
+			adIsRequested = true;
+		}
+
+		if (adIsAvailable && adWantsToBeShown)
+		{
+			if (adPopupWindow.isShowing())
+			{
+				return;
+			}
+
+			adPopupWindow.showAtLocation(activityLayout, adGravity, 0, 0);
+			adPopupWindow.update();
+		}
+		else
+		{
+			if (!adPopupWindow.isShowing())
+			{
+				return;
+			}
+
+			adPopupWindow.dismiss();
+			adPopupWindow.update();
+		}
+	}
+
 	// Called from event thread in NativeActivity	
 	public void AndroidThunkJava_ShowConsoleWindow(String Formats)
 	{
@@ -447,9 +508,17 @@ public class GameActivity extends NativeActivity implements GoogleApiClient.Conn
 
 	public void AndroidThunkJava_GooglePlayConnect()
 	{
-			googleClient.connect();
+		int status = GooglePlayServicesUtil.isGooglePlayServicesAvailable(getApplicationContext());
+
+		// check if google play services is available on this device, or is available with an update
+		if ((status != ConnectionResult.SUCCESS) && (status != ConnectionResult.SERVICE_VERSION_UPDATE_REQUIRED))
+		{
+			return;
 		}
 
+		googleClient.connect();
+	}
+
 	public void AndroidThunkJava_ShowLeaderboard(String LeaderboardID)
 	{
 		Log.debug("In AndroidThunkJava_ShowLeaderboard, ID is " + LeaderboardID);
@@ -543,7 +612,7 @@ public class GameActivity extends NativeActivity implements GoogleApiClient.Conn
 		Log.debug("In AndroidThunkJava_ShowAdBanner");
 		Log.debug("AdID: " + AdMobAdUnitID);
 
-		final int adGravity = bShowOnBottonOfScreen ? Gravity.BOTTOM : Gravity.TOP;
+		adGravity = bShowOnBottonOfScreen ? Gravity.BOTTOM : Gravity.TOP;
 
 		if (adInit)
 		{
@@ -558,8 +627,8 @@ public class GameActivity extends NativeActivity implements GoogleApiClient.Conn
 						return;
 					}
 
-					adPopupWindow.showAtLocation(activityLayout, adGravity, 0, 0);
-					adPopupWindow.update();
+					adWantsToBeShown = true;
+					updateAdVisibility(true);
 				}
 			});
 
@@ -580,16 +649,19 @@ public class GameActivity extends NativeActivity implements GoogleApiClient.Conn
 				{
 					adInit = true;
 
+					final DisplayMetrics dm = getResources().getDisplayMetrics();
+					final float scale = dm.density;
 					adPopupWindow = new PopupWindow(_activity);
-					adPopupWindow.setWidth(320);
-					adPopupWindow.setHeight(50);
+					adPopupWindow.setWidth((int)(320*scale));
+					adPopupWindow.setHeight((int)(50*scale));
 					adPopupWindow.setWindowLayoutMode(LayoutParams.WRAP_CONTENT, LayoutParams.WRAP_CONTENT);
 					adPopupWindow.setClippingEnabled(false);
 
 					adLayout = new LinearLayout(_activity);
 					activityLayout = new LinearLayout(_activity);
 
-					adLayout.setPadding(-5,-5,-5,-5);
+					final int padding = (int)(-5*scale);
+					adLayout.setPadding(padding,padding,padding,padding);
 
 					MarginLayoutParams params = new MarginLayoutParams(LayoutParams.WRAP_CONTENT, LayoutParams.WRAP_CONTENT);;
 
@@ -601,12 +673,31 @@ public class GameActivity extends NativeActivity implements GoogleApiClient.Conn
 
 					_activity.setContentView(activityLayout, params);
 
-					AdRequest adRequest = new AdRequest.Builder().build();		// add test devices here
-
-					_activity.adView.loadAd(adRequest);
-
-					adPopupWindow.showAtLocation(activityLayout, adGravity, 0, 0);
-					adPopupWindow.update();
+					// set up our ad callbacks
+					_activity.adView.setAdListener(new AdListener()
+					{
+						 @Override
+						public void onAdLoaded()
+						{
+							adIsAvailable = true;
+							adIsRequested = false;
+
+							updateAdVisibility(true);
+						}
+
+						 @Override
+						public void onAdFailedToLoad(int errorCode)
+						{
+							adIsAvailable = false;
+							adIsRequested = false;
+
+							// don't immediately request a new ad on failure, wait until the next show
+							updateAdVisibility(false);
+						}
+					});
+
+					adWantsToBeShown = true;
+					updateAdVisibility(true);
 				}
 			});
 		}
@@ -626,13 +717,8 @@ public class GameActivity extends NativeActivity implements GoogleApiClient.Conn
 			@Override
 			public void run()
 			{
-				if ((adPopupWindow == null) || !adPopupWindow.isShowing())
-				{
-					return;
-				}
-
-				adPopupWindow.dismiss();
-				adPopupWindow.update();
+				adWantsToBeShown = false;
+				updateAdVisibility(true);
 			}
 		});
 	}
@@ -652,13 +738,8 @@ public class GameActivity extends NativeActivity implements GoogleApiClient.Conn
 			@Override
 			public void run()
 			{
-				if ((adPopupWindow == null) || !adPopupWindow.isShowing())
-				{
-					return;
-				}
-
-				adPopupWindow.dismiss();
-				adPopupWindow.update();
+				adWantsToBeShown = false;
+				updateAdVisibility(true);
 			}
 		});
 	}
@@ -685,6 +766,7 @@ public class GameActivity extends NativeActivity implements GoogleApiClient.Conn
 	public native void nativeSetWindowInfo(boolean bIsPortrait);
 	public native void nativeSetObbInfo(String PackageName, int Version, int PatchVersion);
 	public native void nativeUpdateAchievements(JavaAchievement[] Achievements);
+	public native void nativeSetAndroidVersionInformation( String AndroidVersion, String PhoneMake, String PhoneModel );
 
 	public native void nativeConsoleCommand(String commandString);
 	
