Integrate [REDACTED_PHONE] -> Main

Fix for google play projects which don't have google play enabled initializing google play..

//depot/UE4-Releases/4.2/Engine/...

to //depot/UE4/Engine/...

diff --git a/Engine/Build/Android/Java/src/com/epicgames/ue4/GameActivity.java b/Engine/Build/Android/Java/src/com/epicgames/ue4/GameActivity.java
--- a/Engine/Build/Android/Java/src/com/epicgames/ue4/GameActivity.java
+++ b/Engine/Build/Android/Java/src/com/epicgames/ue4/GameActivity.java
@@ -508,6 +508,11 @@ public class GameActivity extends NativeActivity implements GoogleApiClient.Conn
 
 	public void AndroidThunkJava_GooglePlayConnect()
 	{
+		if ( !nativeIsGooglePlayEnabled() ) 
+		{
+			return;
+		}
+
 		int status = GooglePlayServicesUtil.isGooglePlayServicesAvailable(getApplicationContext());
 
 		// check if google play services is available on this device, or is available with an update
@@ -516,8 +521,8 @@ public class GameActivity extends NativeActivity implements GoogleApiClient.Conn
 			return;
 		}
 
-		googleClient.connect();
-	}
+			googleClient.connect();
+		}
 
 	public void AndroidThunkJava_ShowLeaderboard(String LeaderboardID)
 	{
@@ -770,6 +775,8 @@ public class GameActivity extends NativeActivity implements GoogleApiClient.Conn
 
 	public native void nativeConsoleCommand(String commandString);
 	
+	public native boolean nativeIsGooglePlayEnabled();
+	
 	static
 	{
 		System.loadLibrary("UE4");
diff --git a/Engine/Source/Runtime/Launch/Private/Android/LaunchAndroid.cpp b/Engine/Source/Runtime/Launch/Private/Android/LaunchAndroid.cpp
--- a/Engine/Source/Runtime/Launch/Private/Android/LaunchAndroid.cpp
+++ b/Engine/Source/Runtime/Launch/Private/Android/LaunchAndroid.cpp
@@ -871,7 +871,7 @@ extern "C" void Java_com_epicgames_ue4_GameActivity_nativeConsoleCommand(JNIEnv*
 //This function is declared in the Java-defined class, GameActivity.java: "public native void nativeIsGooglePlayEnabled();"
 extern "C" jboolean Java_com_epicgames_ue4_GameActivity_nativeIsGooglePlayEnabled(JNIEnv* jenv, jobject thiz)
 {
-	bool bEnabled = true;
+	bool bEnabled = false;
 	GConfig->GetBool(TEXT("/Script/AndroidRuntimeSettings.AndroidRuntimeSettings"), TEXT("bEnableGooglePlaySupport"), bEnabled, GEngineIni);
 	UE_LOG(LogOnline, Log, TEXT("Checking whether Google Play is enabled. bEnableGooglePlaySupport = %d"), bEnabled);
 	return bEnabled;
