Update proguard-project.txt for AdMob on Android according to SDK recommendations

#github https://github.com/EpicGames/UnrealEngine/pull/183 Malfate

diff --git a/Engine/Build/Android/Java/proguard-project.txt b/Engine/Build/Android/Java/proguard-project.txt
--- a/Engine/Build/Android/Java/proguard-project.txt
+++ b/Engine/Build/Android/Java/proguard-project.txt
@@ -43,3 +43,12 @@
 }
 
 # End Google Play services requirements
+
+# Google Play Admob
+-keep public class com.google.android.gms.ads.** {
+   public *;
+}
+
+-keep public class com.google.ads.** {
+   public *;
+}
