- Fixed generic PushCommandLine.bat file

diff --git a/Engine/Build/Android/UE4Game/PushCommandLine.bat b/Engine/Build/Android/UE4Game/PushCommandLine.bat
--- a/Engine/Build/Android/UE4Game/PushCommandLine.bat
+++ b/Engine/Build/Android/UE4Game/PushCommandLine.bat
@@ -1 +1 @@
-adb push UE4CommandLine.txt /mnt/sdcard/UE4Game/UE4CommandLine.txt
\ No newline at end of file
+%ANDROID_HOME%\platform-tools\adb.exe push UE4CommandLine.txt /mnt/sdcard/UE4Game/UE4CommandLine.txt
\ No newline at end of file
