Allow portrait on Android tablets, but force landscape on NVidia Shield, Google TV, Kindle Fire TV, Ouya
#TTP 337401 - CRITICAL: Regression: ANDROID: TappyChicken opens in landscape on tablets
#UE4
#Android

diff --git a/Engine/Build/Android/Java/src/com/epicgames/ue4/GameActivity.java b/Engine/Build/Android/Java/src/com/epicgames/ue4/GameActivity.java
--- a/Engine/Build/Android/Java/src/com/epicgames/ue4/GameActivity.java
+++ b/Engine/Build/Android/Java/src/com/epicgames/ue4/GameActivity.java
@@ -182,10 +182,57 @@ public class GameActivity extends NativeActivity implements GoogleApiClient.Conn
 		// tell Android that we want volume controls to change the media volume, aka music
 		setVolumeControlStream(AudioManager.STREAM_MUSIC);
 		
+		// is this a native landscape device (tablet, tv)?
 		if ( getDeviceDefaultOrientation() == Configuration.ORIENTATION_LANDSCAPE )
 		{
-			Log.debug( "Setting screen orientation to landscape because we have detected landscape device" );
-			_activity.setRequestedOrientation( android.content.pm.ActivityInfo.SCREEN_ORIENTATION_LANDSCAPE );
+			boolean bForceLandscape = false;
+
+			// check for a Google TV by checking system feature support
+			if (getPackageManager().hasSystemFeature("com.google.android.tv"))
+			{
+				Log.debug( "Detected Google TV, will default to landscape" );
+				bForceLandscape = true;
+			} else
+
+			// check NVidia devices
+			if (android.os.Build.MANUFACTURER.equals("NVIDIA"))
+			{
+				// is it a Shield? (checking exact model)
+				if (android.os.Build.MODEL.equals("SHIELD"))
+				{
+					Log.debug( "Detected NVidia Shield, will default to landscape" );
+					bForceLandscape = true;
+				}
+			} else
+
+			// check Ouya
+			if (android.os.Build.MANUFACTURER.equals("OUYA"))
+			{
+				// only one so far (ouya_1_1) but check prefix anyway
+				if (android.os.Build.MODEL.toLowerCase().startsWith("ouya_"))
+				{
+					Log.debug( "Detected Ouya console (" + android.os.Build.MODEL + "), will default to landscape" );
+					bForceLandscape = true;
+				}
+			} else
+
+			// check Amazon devices
+			if (android.os.Build.MANUFACTURER.equals("Amazon"))
+			{
+				// is it a Kindle Fire TV? (Fire TV FAQ says AFTB, but to check for AFT)
+				if (android.os.Build.MODEL.startsWith("AFT"))
+				{
+					Log.debug( "Detected Kindle Fire TV (" + android.os.Build.MODEL + "), will default to landscape" );
+					bForceLandscape = true;
+				}
+			}
+
+			// apply the force request if we found a device above
+			if (bForceLandscape)
+			{
+				Log.debug( "Setting screen orientation to landscape because we have detected landscape device" );
+				_activity.setRequestedOrientation( android.content.pm.ActivityInfo.SCREEN_ORIENTATION_LANDSCAPE );
+			}
 		}
 		
 		// Grab a reference to the asset manager
