Fix Google Play and re-enable it for Tappy Chicken

Don't try to connect() to Google Play before the activity's onCreate() has been called, this was causing a crash/hang.

diff --git a/Engine/Build/Android/Java/src/com/epicgames/ue4/GameActivity.java b/Engine/Build/Android/Java/src/com/epicgames/ue4/GameActivity.java
--- a/Engine/Build/Android/Java/src/com/epicgames/ue4/GameActivity.java
+++ b/Engine/Build/Android/Java/src/com/epicgames/ue4/GameActivity.java
@@ -298,6 +298,9 @@ public class GameActivity extends NativeActivity implements GoogleApiClient.Conn
 		 .addConnectionCallbacks(this)
 		 .addOnConnectionFailedListener(this)
 		 .build();
+
+		// Try to establish a connection to Google Play
+		AndroidThunkJava_GooglePlayConnect();
 	}
 
 	@Override
@@ -559,17 +562,17 @@ public class GameActivity extends NativeActivity implements GoogleApiClient.Conn
 		{
 			return;
 		}
-
-		int status = GooglePlayServicesUtil.isGooglePlayServicesAvailable(getApplicationContext());
-
+		
+		int status = GooglePlayServicesUtil.isGooglePlayServicesAvailable(this);
+		
 		// check if google play services is available on this device, or is available with an update
 		if ((status != ConnectionResult.SUCCESS) && (status != ConnectionResult.SERVICE_VERSION_UPDATE_REQUIRED))
 		{
 			return;
 		}
 
-			googleClient.connect();
-		}
+		googleClient.connect();
+	}
 
 	public void AndroidThunkJava_ShowLeaderboard(String LeaderboardID)
 	{
diff --git a/Engine/Source/Runtime/Online/Android/OnlineSubsystemGooglePlay/Private/OnlineSubsystemGooglePlay.cpp b/Engine/Source/Runtime/Online/Android/OnlineSubsystemGooglePlay/Private/OnlineSubsystemGooglePlay.cpp
--- a/Engine/Source/Runtime/Online/Android/OnlineSubsystemGooglePlay/Private/OnlineSubsystemGooglePlay.cpp
+++ b/Engine/Source/Runtime/Online/Android/OnlineSubsystemGooglePlay/Private/OnlineSubsystemGooglePlay.cpp
@@ -103,9 +103,6 @@ bool FOnlineSubsystemGooglePlay::Init()
 	AchievementsInterface = MakeShareable(new FOnlineAchievementsGooglePlay(this));
 	ExternalUIInterface = MakeShareable(new FOnlineExternalUIGooglePlay());
 
-	extern void AndroidThunkCpp_GooglePlayConnect();
-	AndroidThunkCpp_GooglePlayConnect();
-
 	return true;
 }
 
