Fixed crash on startup on android. Google play structures getting stripped from gameactivity so reflection was failing to initialize in deployment builds.

diff --git a/Engine/Build/Android/Java/proguard-project.txt b/Engine/Build/Android/Java/proguard-project.txt
--- a/Engine/Build/Android/Java/proguard-project.txt
+++ b/Engine/Build/Android/Java/proguard-project.txt
@@ -13,7 +13,13 @@
 # Add any project specific keep options here:
 
 # Needed to keep the JNI functions around
--keepclassmembers class com.epicgames.ue4.GameActivity {
+-keepattributes InnerClasses
+ -keepclassmembers class com.epicgames.ue4.GameActivity {
+    public *;
+ }
+
+# Needed...
+-keep class com.epicgames.ue4.GameActivity$JavaAchievement {
    public *;
 }
 
