Removed some warnings from the packaging process.  Proguard settings were too inclusive.

diff --git a/Engine/Build/Android/Java/proguard-project.txt b/Engine/Build/Android/Java/proguard-project.txt
--- a/Engine/Build/Android/Java/proguard-project.txt
+++ b/Engine/Build/Android/Java/proguard-project.txt
@@ -13,10 +13,9 @@
 # Add any project specific keep options here:
 
 # Needed to keep the JNI functions around
--keepattributes InnerClasses
- -keepclassmembers class com.epicgames.ue4.GameActivity {
-    public *;
- }
+-keepclassmembers class com.epicgames.ue4.GameActivity {
+   public *;
+}
 
 # Needed...
 -keep class com.epicgames.ue4.GameActivity$JavaAchievement {
