Started syncing androidMain with onCreate from java to stop synchronzation issues.  Removed sleep hack which was doing something similar.
Changed the achivement query so that if not connected it will fail right away instead of waiting.

diff --git a/Engine/Build/Android/Java/src/com/epicgames/ue4/GameActivity.java b/Engine/Build/Android/Java/src/com/epicgames/ue4/GameActivity.java
--- a/Engine/Build/Android/Java/src/com/epicgames/ue4/GameActivity.java
+++ b/Engine/Build/Android/Java/src/com/epicgames/ue4/GameActivity.java
@@ -320,6 +320,9 @@ public class GameActivity extends NativeActivity implements GoogleApiClient.Conn
 		 .addOnConnectionFailedListener(this)
 		 .build();
 
+		// Now okay for event handler to be set up on native side
+		nativeResumeMainInit();
+		
 		// Try to establish a connection to Google Play
 		AndroidThunkJava_GooglePlayConnect();
 	}
@@ -490,9 +493,9 @@ public class GameActivity extends NativeActivity implements GoogleApiClient.Conn
 		if (requestCode == GOOGLE_SERVICES_REQUEST_RESOLVE_ERROR)
 		{
 			Log.debug("Google Play Services connection resolution finished with resultCode " + resultCode);
-			bResolvingGoogleServicesError = false;
 			if (resultCode == RESULT_OK)
 			{
+				bResolvingGoogleServicesError = false;
 				// Make sure the app is not already connected or attempting to connect
 				if (!googleClient.isConnecting() &&	!googleClient.isConnected())
 				{
@@ -592,7 +595,10 @@ public class GameActivity extends NativeActivity implements GoogleApiClient.Conn
 			return;
 		}
 
-		googleClient.connect();
+		if ( !googleClient.isConnected() && !googleClient.isConnecting() )
+		{
+			googleClient.connect();
+		}
 	}
 
 	public void AndroidThunkJava_ShowLeaderboard(String LeaderboardID)
@@ -679,8 +685,16 @@ public class GameActivity extends NativeActivity implements GoogleApiClient.Conn
 
 	public void AndroidThunkJava_QueryAchievements()
 	{
-		PendingResult<Achievements.LoadAchievementsResult> loadAchievementsResult = Games.Achievements.load(googleClient, false);
-		loadAchievementsResult.setResultCallback(new QueryAchievementsResultCallback());		
+		//Log.debug("Incremental achievement ID " + AchievementID + ": not setting progress to " + RoundedSteps);
+		if ( googleClient.isConnected() )
+		{
+			PendingResult<Achievements.LoadAchievementsResult> loadAchievementsResult = Games.Achievements.load(googleClient, false);
+			loadAchievementsResult.setResultCallback(new QueryAchievementsResultCallback());		
+		}
+		else
+		{
+			nativeFailedUpdateAchievements();
+		}
 	}
 
 	public void AndroidThunkJava_ShowAdBanner(String AdMobAdUnitID, boolean bShowOnBottonOfScreen)
@@ -842,11 +856,14 @@ public class GameActivity extends NativeActivity implements GoogleApiClient.Conn
 	public native void nativeSetWindowInfo(boolean bIsPortrait, int DepthBufferPreference);
 	public native void nativeSetObbInfo(String PackageName, int Version, int PatchVersion);
 	public native void nativeUpdateAchievements(JavaAchievement[] Achievements);
+	public native void nativeFailedUpdateAchievements();
 	public native void nativeSetAndroidVersionInformation( String AndroidVersion, String PhoneMake, String PhoneModel );
 
 	public native void nativeConsoleCommand(String commandString);
 	
 	public native boolean nativeIsGooglePlayEnabled();
+
+	public native void nativeResumeMainInit();
 	
 	static
 	{
diff --git a/Engine/Source/Runtime/Launch/Private/Android/LaunchAndroid.cpp b/Engine/Source/Runtime/Launch/Private/Android/LaunchAndroid.cpp
--- a/Engine/Source/Runtime/Launch/Private/Android/LaunchAndroid.cpp
+++ b/Engine/Source/Runtime/Launch/Private/Android/LaunchAndroid.cpp
@@ -102,6 +102,20 @@ static const float SampleDecayRate = 0.85f;
 // Event for coordinating pausing of the main and event handling threads to prevent background spinning
 static FEvent* EventHandlerEvent = NULL;
 
+// Wait for Java onCreate to complete before resume main init
+static volatile bool GResumeMainInit = false;
+static volatile bool GEventHandlerInitialized = false;
+extern "C" void Java_com_epicgames_ue4_GameActivity_nativeResumeMainInit(JNIEnv* jenv, jobject thiz)
+{
+	GResumeMainInit = true;
+
+	// now wait for event handler to be set up before returning
+	while (!GEventHandlerInitialized)
+	{
+		FPlatformProcess::Sleep(0.01f);
+		FPlatformMisc::MemoryBarrier();
+	}
+}
 
 
 void UpdateGameInterruptions()
@@ -246,8 +260,12 @@ int32 AndroidMain(struct android_app* state)
 		IgnoredGamepadKeyCodes.Add(IgnoredGamepadKeyCodesList[i]);
 	}
 
-	// wait a moment for the data from the activity to be set
-	FPlatformProcess::Sleep(1.0f);
+	// wait for java activity onCreate to finish
+	while (!GResumeMainInit)
+	{
+		FPlatformProcess::Sleep(0.01f);
+		FPlatformMisc::MemoryBarrier();
+	}
 
 	// read the command line file
 	InitCommandLine();
@@ -257,6 +275,9 @@ int32 AndroidMain(struct android_app* state)
 	FPlatformMisc::LowLevelOutputDebugString(L"Created sync event");
 	FAppEventManager::GetInstance()->SetEventHandlerEvent(EventHandlerEvent);
 
+	// ready for onCreate to complete
+	GEventHandlerInitialized = true;
+
 	// initialize the engine
 	GEngineLoop.PreInit(0, NULL, FCommandLine::Get());
 
@@ -381,7 +402,6 @@ static void AndroidProcessEvents(struct android_app* state)
 	static FVector last_accelerometer(0, 0, 0);
 
 	while((ident = ALooper_pollAll(-1, &fdesc, &events, (void**)&source)) >= 0)
-	// while((ident = ALooper_pollAll(0, &fdesc, &events, (void**)&source)) >= 0)
 	{
 		// process this event
 		if (source)
diff --git a/Engine/Source/Runtime/Online/Android/OnlineSubsystemGooglePlay/Private/OnlineAchievementsInterfaceGooglePlay.cpp b/Engine/Source/Runtime/Online/Android/OnlineSubsystemGooglePlay/Private/OnlineAchievementsInterfaceGooglePlay.cpp
--- a/Engine/Source/Runtime/Online/Android/OnlineSubsystemGooglePlay/Private/OnlineAchievementsInterfaceGooglePlay.cpp
+++ b/Engine/Source/Runtime/Online/Android/OnlineSubsystemGooglePlay/Private/OnlineAchievementsInterfaceGooglePlay.cpp
@@ -68,6 +68,30 @@ extern "C" void Java_com_epicgames_ue4_GameActivity_nativeUpdateAchievements(JNI
 	});
 }
 
+
+
+extern "C" void Java_com_epicgames_ue4_GameActivity_nativeFailedUpdateAchievements(JNIEnv* LocalJNIEnv, jobject LocalThiz )
+{
+	auto AchievementsInterface = FOnlineAchievementsGooglePlay::PendingAchievementQuery.AchievementsInterface;
+
+	if(!AchievementsInterface ||
+		!AchievementsInterface->AndroidSubsystem ||
+		!AchievementsInterface->AndroidSubsystem->GetAsyncTaskManager())
+	{
+		// We should call the delegate with a false parameter here, but if we don't have
+		// the async task manager we're not going to call it on the game thread.
+		return;
+	}
+
+	AchievementsInterface->AndroidSubsystem->GetAsyncTaskManager()->AddGenericToOutQueue([]()
+	{
+		auto& PendingQuery = FOnlineAchievementsGooglePlay::PendingAchievementQuery;
+
+		PendingQuery.Delegate.ExecuteIfBound( PendingQuery.PlayerID, false );
+		PendingQuery.IsQueryPending = false;
+	});
+}
+
 FOnlineAchievementsGooglePlay::FOnlineAchievementsGooglePlay( FOnlineSubsystemGooglePlay* InSubsystem )
 	: AndroidSubsystem(InSubsystem)
 {
diff --git a/Engine/Source/Runtime/Online/Android/OnlineSubsystemGooglePlay/Private/OnlineAchievementsInterfaceGooglePlay.h b/Engine/Source/Runtime/Online/Android/OnlineSubsystemGooglePlay/Private/OnlineAchievementsInterfaceGooglePlay.h
--- a/Engine/Source/Runtime/Online/Android/OnlineSubsystemGooglePlay/Private/OnlineAchievementsInterfaceGooglePlay.h
+++ b/Engine/Source/Runtime/Online/Android/OnlineSubsystemGooglePlay/Private/OnlineAchievementsInterfaceGooglePlay.h
@@ -8,6 +8,7 @@
 #include <jni.h>
 
 extern "C" void Java_com_epicgames_ue4_GameActivity_nativeUpdateAchievements(JNIEnv* LocalJNIEnv, jobject LocalThiz, jobjectArray Achievements);
+extern "C" void Java_com_epicgames_ue4_GameActivity_nativeFailedUpdateAchievements(JNIEnv* LocalJNIEnv, jobject LocalThiz );
 
 /**
  *	IOnlineAchievements - Interface class for Achievements
@@ -52,6 +53,7 @@ private:
 	 * @param Achievements an array of JavaAchievement objects holding the data just queried from the backend
 	 */
 	friend void Java_com_epicgames_ue4_GameActivity_nativeUpdateAchievements(JNIEnv* LocalJNIEnv, jobject LocalThiz, jobjectArray Achievements);
+	friend void Java_com_epicgames_ue4_GameActivity_nativeFailedUpdateAchievements(JNIEnv* LocalJNIEnv, jobject LocalThiz );
 
 public:
 
