Added console swipe history for Android
#Android
#ttp 341718

diff --git a/Engine/Build/Android/Java/src/com/epicgames/ue4/GameActivity.java b/Engine/Build/Android/Java/src/com/epicgames/ue4/GameActivity.java
--- a/Engine/Build/Android/Java/src/com/epicgames/ue4/GameActivity.java
+++ b/Engine/Build/Android/Java/src/com/epicgames/ue4/GameActivity.java
@@ -4,6 +4,7 @@ package com.epicgames.ue4;
 
 import java.util.Map;
 import java.util.HashMap;
+import java.util.ArrayList;
 
 import android.app.NativeActivity;
 import android.os.Bundle;
@@ -12,6 +13,7 @@ import android.util.Log;
 import android.app.AlertDialog;
 import android.app.Dialog;
 import android.widget.EditText;
+import android.text.InputType;
 import android.content.DialogInterface;
 import android.content.Intent;
 import android.content.res.AssetManager;
@@ -25,6 +27,10 @@ import android.media.AudioManager;
 import android.util.DisplayMetrics;
 
 import android.view.Gravity;
+import android.view.MotionEvent;
+import android.view.View;
+import android.view.View.OnTouchListener;
+import android.view.ViewConfiguration;
 import android.view.ViewGroup.LayoutParams;
 import android.view.ViewGroup.MarginLayoutParams;
 import android.view.WindowManager;
@@ -67,6 +73,10 @@ public class GameActivity extends NativeActivity implements GoogleApiClient.Conn
 	// Console
 	AlertDialog consoleAlert;
 	EditText consoleInputBox;
+	ArrayList<String> consoleHistoryList;
+	int consoleHistoryIndex;
+	float consoleDistance;
+	float consoleVelocity;
 
 	// Virtual keyboard
 	AlertDialog virtualKeyboardAlert;
@@ -295,6 +305,58 @@ public class GameActivity extends NativeActivity implements GoogleApiClient.Conn
 		AlertDialog.Builder builder;
 
 		consoleInputBox = new EditText(this);
+		consoleInputBox.setInputType(0x00080001); // TYPE_CLASS_TEXT | TYPE_TEXT_FLAG_NO_SUGGESTIONS);
+		consoleHistoryList = new ArrayList<String>();
+		consoleHistoryIndex = 0;
+
+		final ViewConfiguration vc = ViewConfiguration.get(this);
+        DisplayMetrics dm = getResources().getDisplayMetrics();
+        consoleDistance = vc.getScaledPagingTouchSlop() * dm.density;
+        consoleVelocity = vc.getScaledMinimumFlingVelocity() / 1000.0f;
+
+		consoleInputBox.setOnTouchListener(new OnTouchListener() {
+			private long downTime;
+			private float downX;
+
+			public void swipeLeft() {
+				if (!consoleHistoryList.isEmpty() && consoleHistoryIndex + 1 < consoleHistoryList.size()) {
+					consoleInputBox.setText(consoleHistoryList.get(++consoleHistoryIndex));
+				}
+			}
+
+			public void swipeRight() {
+				if (!consoleHistoryList.isEmpty() && consoleHistoryIndex > 0) {
+					consoleInputBox.setText(consoleHistoryList.get(--consoleHistoryIndex));
+				}
+			}
+
+			public boolean onTouch(View v, MotionEvent event) {
+				switch (event.getAction()) {
+					case MotionEvent.ACTION_DOWN: {
+						// remember down time and position
+						downTime = System.currentTimeMillis();
+						downX = event.getX();
+						return true;
+					}
+					case MotionEvent.ACTION_UP: {
+						long deltaTime = System.currentTimeMillis() - downTime;
+						float delta = event.getX() - downX;
+						float absDelta = Math.abs(delta);
+
+						if (absDelta > consoleDistance && absDelta > deltaTime * consoleVelocity)
+						{
+							if (delta < 0)
+								this.swipeLeft();
+							else
+								this.swipeRight();
+							return true;
+						}
+						return false;
+					}
+				}
+				return false;
+			}
+		});
 
 		builder = new AlertDialog.Builder(this);
 		builder.setTitle("Console Window - Enter Command")
@@ -303,6 +365,15 @@ public class GameActivity extends NativeActivity implements GoogleApiClient.Conn
 		.setPositiveButton("Ok", new DialogInterface.OnClickListener() {
 			public void onClick(DialogInterface dialog, int id) {
 				String message = consoleInputBox.getText().toString().trim();
+
+				// remove it if already in history
+				int index = consoleHistoryList.indexOf(message);
+				if (index >= 0)
+					consoleHistoryList.remove(index);
+
+				// add it to the end
+				consoleHistoryList.add(message);
+
 				nativeConsoleCommand(message);
 				consoleInputBox.setText(" ");
 				dialog.dismiss();
@@ -581,6 +652,9 @@ public class GameActivity extends NativeActivity implements GoogleApiClient.Conn
 			return;
 		}
 
+		// start at end of console history
+		consoleHistoryIndex = consoleHistoryList.size();
+
 		consoleAlert.setMessage("[Available texture formats: " + Formats + "]");
 		_activity.runOnUiThread(new Runnable()
 		{
