UE4 - Fix typo

diff --git a/Engine/Build/Android/Java/src/com/epicgames/ue4/GameActivity.java b/Engine/Build/Android/Java/src/com/epicgames/ue4/GameActivity.java
--- a/Engine/Build/Android/Java/src/com/epicgames/ue4/GameActivity.java
+++ b/Engine/Build/Android/Java/src/com/epicgames/ue4/GameActivity.java
@@ -282,7 +282,7 @@ public class GameActivity extends NativeActivity implements GoogleApiClient.Conn
 
 
 		Log.debug( "Android version is " + android.os.Build.VERSION.RELEASE );
-		Log.debug( "Android manufactorer is " + android.os.Build.MANUFACTURER );
+		Log.debug( "Android manufacturer is " + android.os.Build.MANUFACTURER );
 		Log.debug( "Android model is " + android.os.Build.MODEL );
 
 		nativeSetAndroidVersionInformation( android.os.Build.VERSION.RELEASE, android.os.Build.MANUFACTURER, android.os.Build.MODEL );
