Add new console command "online.ResetAchievements" and implement ResetAchievements for Google Play

On Google Play, this command will reset the achievements of the currently signed in user, and this user must be flagged as tester account in the developer console.
Google only seems to be resetting standard (non-incremental) achievements for some reason.

diff --git a/Engine/Build/Android/Java/src/com/epicgames/ue4/GameActivity.java b/Engine/Build/Android/Java/src/com/epicgames/ue4/GameActivity.java
--- a/Engine/Build/Android/Java/src/com/epicgames/ue4/GameActivity.java
+++ b/Engine/Build/Android/Java/src/com/epicgames/ue4/GameActivity.java
@@ -37,6 +37,7 @@ import android.view.WindowManager;
 import android.widget.LinearLayout;
 import android.widget.PopupWindow;
 
+import com.google.android.gms.auth.GoogleAuthUtil;
 import com.google.android.gms.common.api.GoogleApiClient;
 import com.google.android.gms.common.api.PendingResult;
 import com.google.android.gms.common.api.ResultCallback;
@@ -50,6 +51,11 @@ import com.google.android.gms.ads.AdView;
 import com.google.android.gms.ads.AdSize;
 import com.google.android.gms.ads.AdListener;
 
+import com.google.android.gms.plus.Plus;
+
+import java.net.URL;
+import java.net.HttpURLConnection;
+
 // TODO: use the resources from the UE4 lib project once we've got the packager up and running
 //import com.epicgames.ue4.R;
 import com.epicgames.ue4.JavaBuildSettings;
@@ -413,6 +419,8 @@ public class GameActivity extends NativeActivity implements GoogleApiClient.Conn
 		googleClient = new GoogleApiClient.Builder(this)
 		 .addApi(Games.API)
 		 .addScope(Games.SCOPE_GAMES)
+		 .addApi(Plus.API, null)
+		 .addScope(Plus.SCOPE_PLUS_PROFILE)
 		 .addConnectionCallbacks(this)
 		 .addOnConnectionFailedListener(this)
 		 .build();
@@ -827,6 +835,38 @@ public class GameActivity extends NativeActivity implements GoogleApiClient.Conn
 		}
 	}
 
+	public void AndroidThunkJava_ResetAchievements()
+	{
+		try
+        {
+			String email = Plus.AccountApi.getAccountName(googleClient);
+			Log.debug("AndroidThunkJava_ResetAchievements: using email " + email);
+
+            String accesstoken = GoogleAuthUtil.getToken(this, email, "oauth2:https://www.googleapis.com/auth/games");
+
+			String ResetURL = "https://www.googleapis.com/games/v1management/achievements/reset?access_token=" + accesstoken;
+			Log.debug("AndroidThunkJava_ResetAchievements: using URL " + ResetURL);
+
+			URL url = new URL(ResetURL);
+			HttpURLConnection urlConnection = (HttpURLConnection)url.openConnection();
+
+			try
+			{
+				urlConnection.setRequestMethod("POST");
+				int status = urlConnection.getResponseCode();
+				Log.debug("AndroidThunkJava_ResetAchievements: HTTP response is " + status);
+			}
+			finally
+			{
+				urlConnection.disconnect();
+			}
+        }
+        catch(Exception e)
+        {
+            Log.debug("AndroidThunkJava_ResetAchievements failed: " + e.getMessage());
+        }
+	}
+
 	public void AndroidThunkJava_ShowAdBanner(String AdMobAdUnitID, boolean bShowOnBottonOfScreen)
 	{
 		Log.debug("In AndroidThunkJava_ShowAdBanner");
diff --git a/Engine/Source/Runtime/Launch/Private/Android/AndroidJNI.cpp b/Engine/Source/Runtime/Launch/Private/Android/AndroidJNI.cpp
--- a/Engine/Source/Runtime/Launch/Private/Android/AndroidJNI.cpp
+++ b/Engine/Source/Runtime/Launch/Private/Android/AndroidJNI.cpp
@@ -110,6 +110,7 @@ jmethodID JDef_GameActivity::AndroidThunkJava_LaunchURL;
 jmethodID JDef_GameActivity::AndroidThunkJava_ShowLeaderboard;
 jmethodID JDef_GameActivity::AndroidThunkJava_ShowAchievements;
 jmethodID JDef_GameActivity::AndroidThunkJava_QueryAchievements;
+jmethodID JDef_GameActivity::AndroidThunkJava_ResetAchievements;
 jmethodID JDef_GameActivity::AndroidThunkJava_WriteLeaderboardValue;
 jmethodID JDef_GameActivity::AndroidThunkJava_GooglePlayConnect;
 jmethodID JDef_GameActivity::AndroidThunkJava_WriteAchievement;
@@ -242,6 +243,14 @@ void AndroidThunkCpp_ShowAchievements()
  	}
 }
 
+void AndroidThunkCpp_ResetAchievements()
+{
+	if (JNIEnv* Env = GetJavaEnv())
+	{
+		Env->CallVoidMethod(GJavaGlobalThis, JDef_GameActivity::AndroidThunkJava_ResetAchievements);
+	}
+}
+
 void AndroidThunkCpp_WriteLeaderboardValue(const FString& LeaderboardName, int64_t Value)
 {
 	if (JNIEnv* Env = GetJavaEnv())
@@ -355,6 +364,7 @@ JNIEXPORT jint JNI_OnLoad(JavaVM* InJavaVM, void* InReserved)
 	JDef_GameActivity::AndroidThunkJava_ShowLeaderboard = env->GetMethodID(JDef_GameActivity::ClassID, "AndroidThunkJava_ShowLeaderboard", "(Ljava/lang/String;)V");
 	JDef_GameActivity::AndroidThunkJava_ShowAchievements = env->GetMethodID(JDef_GameActivity::ClassID, "AndroidThunkJava_ShowAchievements", "()V");
 	JDef_GameActivity::AndroidThunkJava_QueryAchievements = env->GetMethodID(JDef_GameActivity::ClassID, "AndroidThunkJava_QueryAchievements", "()V");
+	JDef_GameActivity::AndroidThunkJava_ResetAchievements = env->GetMethodID(JDef_GameActivity::ClassID, "AndroidThunkJava_ResetAchievements", "()V");
 	JDef_GameActivity::AndroidThunkJava_WriteLeaderboardValue = env->GetMethodID(JDef_GameActivity::ClassID, "AndroidThunkJava_WriteLeaderboardValue", "(Ljava/lang/String;J)V");
 	JDef_GameActivity::AndroidThunkJava_GooglePlayConnect = env->GetMethodID(JDef_GameActivity::ClassID, "AndroidThunkJava_GooglePlayConnect", "()V");
 	JDef_GameActivity::AndroidThunkJava_WriteAchievement = env->GetMethodID(JDef_GameActivity::ClassID, "AndroidThunkJava_WriteAchievement", "(Ljava/lang/String;F)V");
diff --git a/Engine/Source/Runtime/Launch/Private/Android/AndroidJNI.h b/Engine/Source/Runtime/Launch/Private/Android/AndroidJNI.h
--- a/Engine/Source/Runtime/Launch/Private/Android/AndroidJNI.h
+++ b/Engine/Source/Runtime/Launch/Private/Android/AndroidJNI.h
@@ -25,6 +25,7 @@ public:
 	static jmethodID AndroidThunkJava_ShowLeaderboard;
 	static jmethodID AndroidThunkJava_ShowAchievements;
 	static jmethodID AndroidThunkJava_QueryAchievements;
+	static jmethodID AndroidThunkJava_ResetAchievements;
 	static jmethodID AndroidThunkJava_WriteLeaderboardValue;
 	static jmethodID AndroidThunkJava_GooglePlayConnect;
 	static jmethodID AndroidThunkJava_WriteAchievement;
diff --git a/Engine/Source/Runtime/Online/Android/OnlineSubsystemGooglePlay/Private/OnlineAchievementsInterfaceGooglePlay.cpp b/Engine/Source/Runtime/Online/Android/OnlineSubsystemGooglePlay/Private/OnlineAchievementsInterfaceGooglePlay.cpp
--- a/Engine/Source/Runtime/Online/Android/OnlineSubsystemGooglePlay/Private/OnlineAchievementsInterfaceGooglePlay.cpp
+++ b/Engine/Source/Runtime/Online/Android/OnlineSubsystemGooglePlay/Private/OnlineAchievementsInterfaceGooglePlay.cpp
@@ -198,7 +198,11 @@ EOnlineCachedResult::Type FOnlineAchievementsGooglePlay::GetCachedAchievementDes
 #if !UE_BUILD_SHIPPING
 bool FOnlineAchievementsGooglePlay::ResetAchievements( const FUniqueNetId& PlayerId )
 {
-	check(!TEXT("ResetAchievements has not been implemented"));
+	UE_LOG(LogOnline, Log, TEXT("Resetting Google Play achievements."));
+
+	extern void AndroidThunkCpp_ResetAchievements();
+	AndroidThunkCpp_ResetAchievements();
+
 	return false;
 };
 #endif // !UE_BUILD_SHIPPING
diff --git a/Engine/Source/Runtime/Online/OnlineSubsystem/Private/OnlineSubsystem.cpp b/Engine/Source/Runtime/Online/OnlineSubsystem/Private/OnlineSubsystem.cpp
--- a/Engine/Source/Runtime/Online/OnlineSubsystem/Private/OnlineSubsystem.cpp
+++ b/Engine/Source/Runtime/Online/OnlineSubsystem/Private/OnlineSubsystem.cpp
@@ -3,7 +3,10 @@
 #include "OnlineSubsystemPrivatePCH.h"
 #include "OnlineSessionInterface.h"
 #include "OnlineIdentityInterface.h"
+#include "OnlineAchievementsInterface.h"
 #include "NboSerializer.h"
+#include "Online.h"
+#include "IConsoleManager.h"
 
 DEFINE_LOG_CATEGORY(LogOnline);
 DEFINE_LOG_CATEGORY(LogOnlineGame);
@@ -86,3 +89,25 @@ bool IsPlayerInSessionImpl(IOnlineSession* SessionInt, FName SessionName, const
 	return bFound;
 }
 
+
+#if !UE_BUILD_SHIPPING
+
+static void ResetAchievements()
+{
+	auto AchievementsInterface = Online::GetAchievementsInterface();
+	if (!AchievementsInterface.IsValid())
+	{
+		UE_LOG_ONLINE(Warning, TEXT("ResetAchievements command: couldn't get the achievements interface"));
+		return;
+	}
+
+	AchievementsInterface->ResetAchievements(FUniqueNetIdString());
+}
+
+FAutoConsoleCommand CmdResetAchievements(
+	TEXT("online.ResetAchievements"),
+	TEXT("Reset achievements for the currently logged in user."),
+	FConsoleCommandDelegate::CreateStatic(ResetAchievements)
+	);
+
+#endif
