Add Android GET_ACCOUNTS permission, needed to reset Google Play achievements

diff --git a/Engine/Build/Android/Java/AndroidManifest.xml b/Engine/Build/Android/Java/AndroidManifest.xml
--- a/Engine/Build/Android/Java/AndroidManifest.xml
+++ b/Engine/Build/Android/Java/AndroidManifest.xml
@@ -41,4 +41,5 @@
 	<uses-permission android:name="com.android.vending.CHECK_LICENSE" />
 	<uses-permission android:name="android.permission.ACCESS_WIFI_STATE"/>
 	<uses-permission android:name="android.permission.MODIFY_AUDIO_SETTINGS"/>
+	<uses-permission android:name="android.permission.GET_ACCOUNTS" />
 </manifest>
