TPP #341357 MobileMVP: Add support for OSS connect action in blueprint and response handling
Initial pass
Adds component for blueprint and interface
Contains Android implementation only.
Everything else is null stubbed out.

Todo: figure out how to deal with failure error codes properly

diff --git a/Engine/Build/Android/Java/src/com/epicgames/ue4/GameActivity.java b/Engine/Build/Android/Java/src/com/epicgames/ue4/GameActivity.java
--- a/Engine/Build/Android/Java/src/com/epicgames/ue4/GameActivity.java
+++ b/Engine/Build/Android/Java/src/com/epicgames/ue4/GameActivity.java
@@ -417,19 +417,19 @@ public class GameActivity extends NativeActivity implements GoogleApiClient.Conn
 
 		// Connect to Google Play Services
 		googleClient = new GoogleApiClient.Builder(this)
+		 .addConnectionCallbacks(this)
+		 .addOnConnectionFailedListener(this)
 		 .addApi(Games.API)
 		 .addScope(Games.SCOPE_GAMES)
 		 .addApi(Plus.API, null)
 		 .addScope(Plus.SCOPE_PLUS_PROFILE)
-		 .addConnectionCallbacks(this)
-		 .addOnConnectionFailedListener(this)
 		 .build();
 
 		// Now okay for event handler to be set up on native side
 		nativeResumeMainInit();
 		
 		// Try to establish a connection to Google Play
-		AndroidThunkJava_GooglePlayConnect();
+		// AndroidThunkJava_GooglePlayConnect();
 
 		Log.debug("==============> GameActive.onCreate complete!");
 	}
@@ -532,6 +532,8 @@ public class GameActivity extends NativeActivity implements GoogleApiClient.Conn
 		// Set the flag that we successfully connected. Checked in onStart to re-establish the connection.
 		bHaveConnectedToGooglePlay = true;
 
+		nativeCompletedConnection(RESULT_OK);
+
 		// Load achievements. Since games are expected to pass in achievement progress as a percentage,
 		// we need to know what the maximum steps are in order to convert the percentage to an integer
 		// number of steps.
@@ -556,6 +558,7 @@ public class GameActivity extends NativeActivity implements GoogleApiClient.Conn
 		if (bResolvingGoogleServicesError)
 		{
 			// Already attempting to resolve an error.
+			Log.debug("... and already trying to resolve an error.");
 			return;
 		}
 		else if (result.hasResolution())
@@ -600,15 +603,23 @@ public class GameActivity extends NativeActivity implements GoogleApiClient.Conn
 		if (requestCode == GOOGLE_SERVICES_REQUEST_RESOLVE_ERROR)
 		{
 			Log.debug("Google Play Services connection resolution finished with resultCode " + resultCode);
-			if (resultCode == RESULT_OK)
+			
+			bResolvingGoogleServicesError = false;
+
+			if (resultCode == RESULT_OK) // -1
 			{
-				bResolvingGoogleServicesError = false;
 				// Make sure the app is not already connected or attempting to connect
 				if (!googleClient.isConnecting() &&	!googleClient.isConnected())
 				{
 					googleClient.connect();
 				}
 			}
+			else
+			{
+				// translate result code? 
+				// 0 if we cancel out the attempt to error recover...
+				nativeCompletedConnection(resultCode);
+			}
 		}
 	}
 
@@ -730,6 +741,7 @@ public class GameActivity extends NativeActivity implements GoogleApiClient.Conn
 		// check if google play services is available on this device, or is available with an update
 		if ((status != ConnectionResult.SUCCESS) && (status != ConnectionResult.SERVICE_VERSION_UPDATE_REQUIRED))
 		{
+			nativeCompletedConnection(status);
 			return;
 		}
 
@@ -1039,6 +1051,8 @@ public class GameActivity extends NativeActivity implements GoogleApiClient.Conn
 
 	public native void nativeResumeMainInit();
 	
+	public native void nativeCompletedConnection(int errorCode);
+
 	static
 	{
 		System.loadLibrary("UE4");
diff --git a/Engine/Source/Runtime/Online/Android/OnlineSubsystemGooglePlay/Private/OnlineConnectionInterfaceGooglePlay.cpp b/Engine/Source/Runtime/Online/Android/OnlineSubsystemGooglePlay/Private/OnlineConnectionInterfaceGooglePlay.cpp
new file mode 100644
--- /dev/null
+++ b/Engine/Source/Runtime/Online/Android/OnlineSubsystemGooglePlay/Private/OnlineConnectionInterfaceGooglePlay.cpp
@@ -0,0 +1,58 @@
+// Copyright 1998-2014 Epic Games, Inc. All Rights Reserved.
+
+#include "OnlineSubsystemGooglePlayPrivatePCH.h"
+#include "OnlineAsyncTaskManagerGooglePlay.h"
+#include "Online.h"
+#include "Android/AndroidJNI.h"
+
+FOnlineConnectionGooglePlay::FPendingConnection FOnlineConnectionGooglePlay::PendingConnectRequest;
+
+// Java interface to deal with callbacks
+extern "C" void Java_com_epicgames_ue4_GameActivity_nativeCompletedConnection(JNIEnv* LocalJNIEnv, jobject LocalThiz, jint errorCode)
+{
+	auto ConnectionInterface = FOnlineConnectionGooglePlay::PendingConnectRequest.ConnectionInterface;
+
+	if (!ConnectionInterface ||
+		!ConnectionInterface->AndroidSubsystem ||
+		!ConnectionInterface->AndroidSubsystem->GetAsyncTaskManager())
+	{
+		// We should call the delegate with a false parameter here, but if we don't have
+		// the async task manager we're not going to call it on the game thread.
+		return;
+	}
+
+	ConnectionInterface->AndroidSubsystem->GetAsyncTaskManager()->AddGenericToOutQueue([errorCode]()
+	{
+		auto& PendingConnection = FOnlineConnectionGooglePlay::PendingConnectRequest;
+		bool failed = errorCode == -1; // find out what a 'no fail' error code might be...
+
+		PendingConnection.Delegate.ExecuteIfBound(errorCode, failed);
+		PendingConnection.Delegate.Unbind();
+		PendingConnection.IsConnectionPending = false;
+	});
+
+}
+
+FOnlineConnectionGooglePlay::FOnlineConnectionGooglePlay( FOnlineSubsystemGooglePlay* InSubsystem )
+	: AndroidSubsystem(InSubsystem)
+{
+	PendingConnectRequest.ConnectionInterface = this;
+}
+
+void FOnlineConnectionGooglePlay::Connect(const FOnConnectionCompleteDelegate & Delegate)
+{
+	if (PendingConnectRequest.IsConnectionPending)
+	{
+		UE_LOG(LogOnline, Warning, TEXT("FOnlineConnectionGooglePlay::Connect: Cannot start a new connection attempt while one is in progress."));
+		Delegate.ExecuteIfBound(0,false); // Need an error code here...
+		return;
+	}
+
+	PendingConnectRequest.Delegate = Delegate;
+	PendingConnectRequest.IsConnectionPending = true;
+
+	// Kick off the connect process
+	extern void AndroidThunkCpp_GooglePlayConnect();
+	FPlatformMisc::LowLevelOutputDebugStringf(TEXT("In FOnlineConnectionGooglePlay::Connect()"));
+	AndroidThunkCpp_GooglePlayConnect();
+}
\ No newline at end of file
diff --git a/Engine/Source/Runtime/Online/Android/OnlineSubsystemGooglePlay/Private/OnlineConnectionInterfaceGooglePlay.h b/Engine/Source/Runtime/Online/Android/OnlineSubsystemGooglePlay/Private/OnlineConnectionInterfaceGooglePlay.h
new file mode 100644
--- /dev/null
+++ b/Engine/Source/Runtime/Online/Android/OnlineSubsystemGooglePlay/Private/OnlineConnectionInterfaceGooglePlay.h
@@ -0,0 +1,57 @@
+// Copyright 1998-2014 Epic Games, Inc. All Rights Reserved.
+
+#pragma once
+
+#include "OnlineConnectionInterface.h"
+#include "OnlineSubsystemTypes.h"
+#include "AndroidRuntimeSettings.h"
+#include <jni.h>
+
+extern "C" void Java_com_epicgames_ue4_GameActivity_nativeCompletedConnection(JNIEnv* LocalJNIEnv, jobject LocalThiz, jint errorCode);
+
+/**
+ *	IOnlineConnection - Interface class for connection management
+ */
+class FOnlineConnectionGooglePlay : public IOnlineConnection
+{
+private:
+
+	/** Reference to the main GameCenter subsystem */
+	class FOnlineSubsystemGooglePlay* AndroidSubsystem;
+	
+	/** hide the default constructor, we need a reference to our OSS */
+	FOnlineConnectionGooglePlay() {}
+
+	
+	/**
+	 * Native function called from Java when we get the connection result callback.
+	 *
+	 * @param LocalJNIEnv the current JNI environment
+	 * @param LocalThiz instance of the Java GameActivity class
+	 * @param errorCode The reported error code from the connection attempt
+	 */
+	friend void Java_com_epicgames_ue4_GameActivity_nativeCompletedConnection(JNIEnv* LocalJNIEnv, jobject LocalThiz, jint errorCode);
+
+	struct FPendingConnection
+	{
+		FOnlineConnectionGooglePlay * ConnectionInterface;
+		FOnConnectionCompleteDelegate Delegate;
+		bool IsConnectionPending;
+	};
+
+
+	/** Instance of the connection context data */
+	static FPendingConnection PendingConnectRequest;
+public:
+
+	/**
+	 * Constructor
+	 *
+	 * @param InSubsystem - A reference to the owning subsystem
+	 */
+	FOnlineConnectionGooglePlay( class FOnlineSubsystemGooglePlay* InSubsystem );
+
+	virtual void Connect(const FOnConnectionCompleteDelegate & Delegate = FOnConnectionCompleteDelegate()) override;
+};
+
+typedef TSharedPtr<FOnlineConnectionGooglePlay, ESPMode::ThreadSafe> FOnlineConnectionGooglePlayPtr;
diff --git a/Engine/Source/Runtime/Online/Android/OnlineSubsystemGooglePlay/Private/OnlineSubsystemGooglePlay.cpp b/Engine/Source/Runtime/Online/Android/OnlineSubsystemGooglePlay/Private/OnlineSubsystemGooglePlay.cpp
--- a/Engine/Source/Runtime/Online/Android/OnlineSubsystemGooglePlay/Private/OnlineSubsystemGooglePlay.cpp
+++ b/Engine/Source/Runtime/Online/Android/OnlineSubsystemGooglePlay/Private/OnlineSubsystemGooglePlay.cpp
@@ -102,6 +102,7 @@ bool FOnlineSubsystemGooglePlay::Init()
 	LeaderboardsInterface = MakeShareable(new FOnlineLeaderboardsGooglePlay(this));
 	AchievementsInterface = MakeShareable(new FOnlineAchievementsGooglePlay(this));
 	ExternalUIInterface = MakeShareable(new FOnlineExternalUIGooglePlay());
+	ConnectionInterface = MakeShareable(new FOnlineConnectionGooglePlay(this));
 
 	return true;
 }
@@ -133,6 +134,7 @@ bool FOnlineSubsystemGooglePlay::Shutdown()
 	DESTRUCT_INTERFACE(AchievementsInterface);
 	DESTRUCT_INTERFACE(LeaderboardsInterface);
 	DESTRUCT_INTERFACE(IdentityInterface);
+	DESTRUCT_INTERFACE(ConnectionInterface);
 #undef DESTRUCT_INTERFACE
 
 	OnlineAsyncTaskThreadRunnable.Reset();
diff --git a/Engine/Source/Runtime/Online/Android/OnlineSubsystemGooglePlay/Public/OnlineSubsystemGooglePlay.h b/Engine/Source/Runtime/Online/Android/OnlineSubsystemGooglePlay/Public/OnlineSubsystemGooglePlay.h
--- a/Engine/Source/Runtime/Online/Android/OnlineSubsystemGooglePlay/Public/OnlineSubsystemGooglePlay.h
+++ b/Engine/Source/Runtime/Online/Android/OnlineSubsystemGooglePlay/Public/OnlineSubsystemGooglePlay.h
@@ -7,6 +7,7 @@
 #include "OnlineAchievementsInterfaceGooglePlay.h"
 #include "OnlineLeaderboardInterfaceGooglePlay.h"
 #include "OnlineExternalUIInterfaceGooglePlay.h"
+#include "OnlineConnectionInterfaceGooglePlay.h"
 #include "UniquePtr.h"
 
 /**
@@ -40,6 +41,8 @@ public:
 	virtual IOnlineUserPtr GetUserInterface() const override { return NULL; }
 	virtual IOnlineAchievementsPtr GetAchievementsInterface() const override;
 	virtual IOnlinePresencePtr GetPresenceInterface() const override { return NULL; }
+	virtual IOnlineConnectionPtr GetConnectionInterface() const override { return ConnectionInterface;  }
+
 	virtual class UObject* GetNamedInterface(FName InterfaceName) override { return NULL; }
 	virtual void SetNamedInterface(FName InterfaceName, class UObject* NewInterface) override {}
 	virtual bool IsDedicated() const override { return false; }
@@ -82,6 +85,9 @@ private:
 
 	/** Interface to the external UI services */
 	FOnlineExternalUIGooglePlayPtr ExternalUIInterface;
+
+	/** Interface to the connection system */
+	FOnlineConnectionGooglePlayPtr ConnectionInterface;
 };
 
 typedef TSharedPtr<FOnlineSubsystemGooglePlay, ESPMode::ThreadSafe> FOnlineSubsystemGooglePlayPtr;
diff --git a/Engine/Source/Runtime/Online/IOS/OnlineSubsystemIOS/Private/OnlineSubsystemIOS.cpp b/Engine/Source/Runtime/Online/IOS/OnlineSubsystemIOS/Private/OnlineSubsystemIOS.cpp
--- a/Engine/Source/Runtime/Online/IOS/OnlineSubsystemIOS/Private/OnlineSubsystemIOS.cpp
+++ b/Engine/Source/Runtime/Online/IOS/OnlineSubsystemIOS/Private/OnlineSubsystemIOS.cpp
@@ -117,6 +117,11 @@ IOnlinePresencePtr FOnlineSubsystemIOS::GetPresenceInterface() const
 	return NULL;
 }
 
+IOnlineConnectionPtr FOnlineSubsystemIOS::GetConnectionInterface() const
+{
+	return NULL;
+}
+
 bool FOnlineSubsystemIOS::Init() 
 {
 	bool bSuccessfullyStartedUp = true;
diff --git a/Engine/Source/Runtime/Online/IOS/OnlineSubsystemIOS/Public/OnlineSubsystemIOS.h b/Engine/Source/Runtime/Online/IOS/OnlineSubsystemIOS/Public/OnlineSubsystemIOS.h
--- a/Engine/Source/Runtime/Online/IOS/OnlineSubsystemIOS/Public/OnlineSubsystemIOS.h
+++ b/Engine/Source/Runtime/Online/IOS/OnlineSubsystemIOS/Public/OnlineSubsystemIOS.h
@@ -39,6 +39,7 @@ public:
 	virtual IOnlineUserPtr GetUserInterface() const override;
 	virtual IOnlineMessagePtr GetMessageInterface() const override;
 	virtual IOnlinePresencePtr GetPresenceInterface() const override;
+	virtual IOnlineConnectionPtr GetConnectionInterface() const override;
 	virtual bool Init() override;
 	virtual bool Shutdown() override;
 	virtual FString GetAppId() const override;
diff --git a/Engine/Source/Runtime/Online/OnlineSubsystem/Public/Interfaces/OnlineConnectionInterface.h b/Engine/Source/Runtime/Online/OnlineSubsystem/Public/Interfaces/OnlineConnectionInterface.h
new file mode 100644
--- /dev/null
+++ b/Engine/Source/Runtime/Online/OnlineSubsystem/Public/Interfaces/OnlineConnectionInterface.h
@@ -0,0 +1,27 @@
+// Copyright 1998-2014 Epic Games, Inc. All Rights Reserved.
+
+
+#pragma once
+
+
+#include "OnlineDelegateMacros.h"
+#include "OnlineStats.h"
+
+
+/**
+* Delegate fired when a connection attempted has completed
+*/
+DECLARE_DELEGATE_TwoParams(FOnConnectionCompleteDelegate, const int, const bool); // TODO: pass down status messages...
+
+
+/**
+ *	IOnlineConnection - Interface class for connection management
+ */
+class IOnlineConnection
+{
+public:
+	virtual ~IOnlineConnection() {}
+
+	//@TODO: Pure virtual?
+	virtual void Connect(const FOnConnectionCompleteDelegate & Delegate = FOnConnectionCompleteDelegate()) = 0;
+};
\ No newline at end of file
diff --git a/Engine/Source/Runtime/Online/OnlineSubsystem/Public/OnlineSubsystem.h b/Engine/Source/Runtime/Online/OnlineSubsystem/Public/OnlineSubsystem.h
--- a/Engine/Source/Runtime/Online/OnlineSubsystem/Public/OnlineSubsystem.h
+++ b/Engine/Source/Runtime/Online/OnlineSubsystem/Public/OnlineSubsystem.h
@@ -50,6 +50,7 @@ typedef TSharedPtr<class IOnlineUser, ESPMode::ThreadSafe> IOnlineUserPtr;
 typedef TSharedPtr<class IOnlineMessage, ESPMode::ThreadSafe> IOnlineMessagePtr;
 typedef TSharedPtr<class IOnlinePresence, ESPMode::ThreadSafe> IOnlinePresencePtr;
 typedef TSharedPtr<class IOnlineParty, ESPMode::ThreadSafe> IOnlinePartyPtr;
+typedef TSharedPtr<class IOnlineConnection, ESPMode::ThreadSafe> IOnlineConnectionPtr;
 
 /**
  * Called when the connection state as reported by the online platform changes
@@ -246,6 +247,12 @@ public:
 	 */
 	virtual class UObject* GetNamedInterface(FName InterfaceName) = 0;
 
+	/**
+	* Get the interface for managing the connection to an online service
+	* @return Interface pointer for the appropriate online service
+	*/
+	virtual IOnlineConnectionPtr GetConnectionInterface() const = 0;
+
 	/**
 	 * Set a custom UObject to be preserved by the online subsystem
 	 *
diff --git a/Engine/Source/Runtime/Online/OnlineSubsystemAmazon/Public/OnlineSubsystemAmazon.h b/Engine/Source/Runtime/Online/OnlineSubsystemAmazon/Public/OnlineSubsystemAmazon.h
--- a/Engine/Source/Runtime/Online/OnlineSubsystemAmazon/Public/OnlineSubsystemAmazon.h
+++ b/Engine/Source/Runtime/Online/OnlineSubsystemAmazon/Public/OnlineSubsystemAmazon.h
@@ -117,6 +117,11 @@ public:
 		return NULL;
 	}
 
+	IOnlineConnectionPtr GetConnectionInterface() const override
+	{
+		return NULL;
+	}
+
 	virtual bool Init() override;
 	virtual bool Shutdown() override;
 	virtual FString GetAppId() const override;
diff --git a/Engine/Source/Runtime/Online/OnlineSubsystemFacebook/Private/IOS/OnlineSubsystemFacebook.cpp b/Engine/Source/Runtime/Online/OnlineSubsystemFacebook/Private/IOS/OnlineSubsystemFacebook.cpp
--- a/Engine/Source/Runtime/Online/OnlineSubsystemFacebook/Private/IOS/OnlineSubsystemFacebook.cpp
+++ b/Engine/Source/Runtime/Online/OnlineSubsystemFacebook/Private/IOS/OnlineSubsystemFacebook.cpp
@@ -116,6 +116,11 @@ IOnlinePartyPtr FOnlineSubsystemFacebook::GetPartyInterface() const
 	return NULL;
 }
 
+IOnlineConnectionPtr FOnlineSubsystemFacebook::GetConnectionInterface() const
+{
+	return NULL;
+}
+
 bool FOnlineSubsystemFacebook::Init() 
 {
 	bool bSuccessfullyStartedUp = true;
diff --git a/Engine/Source/Runtime/Online/OnlineSubsystemFacebook/Private/Windows/OnlineSubsystemFacebook.cpp b/Engine/Source/Runtime/Online/OnlineSubsystemFacebook/Private/Windows/OnlineSubsystemFacebook.cpp
--- a/Engine/Source/Runtime/Online/OnlineSubsystemFacebook/Private/Windows/OnlineSubsystemFacebook.cpp
+++ b/Engine/Source/Runtime/Online/OnlineSubsystemFacebook/Private/Windows/OnlineSubsystemFacebook.cpp
@@ -114,6 +114,11 @@ IOnlinePartyPtr FOnlineSubsystemFacebook::GetPartyInterface() const
 	return NULL;
 }
 
+IOnlineConnectionPtr FOnlineSubsystemFacebook::GetConnectionInterface() const
+{
+	return NULL;
+}
+
 bool FOnlineSubsystemFacebook::Tick(float DeltaTime)
 {
 	if (FacebookIdentity.IsValid())
diff --git a/Engine/Source/Runtime/Online/OnlineSubsystemFacebook/Public/OnlineSubsystemFacebook.h b/Engine/Source/Runtime/Online/OnlineSubsystemFacebook/Public/OnlineSubsystemFacebook.h
--- a/Engine/Source/Runtime/Online/OnlineSubsystemFacebook/Public/OnlineSubsystemFacebook.h
+++ b/Engine/Source/Runtime/Online/OnlineSubsystemFacebook/Public/OnlineSubsystemFacebook.h
@@ -41,6 +41,8 @@ public:
 	virtual IOnlineUserPtr GetUserInterface() const override;
 	virtual IOnlineMessagePtr GetMessageInterface() const override;
 	virtual IOnlinePresencePtr GetPresenceInterface() const override;
+	virtual IOnlineConnectionPtr GetConnectionInterface() const override;
+
 	virtual bool Init() override;
 	virtual bool Shutdown() override;
 	virtual FString GetAppId() const override;
diff --git a/Engine/Source/Runtime/Online/OnlineSubsystemNull/Private/OnlineSubsystemNull.cpp b/Engine/Source/Runtime/Online/OnlineSubsystemNull/Private/OnlineSubsystemNull.cpp
--- a/Engine/Source/Runtime/Online/OnlineSubsystemNull/Private/OnlineSubsystemNull.cpp
+++ b/Engine/Source/Runtime/Online/OnlineSubsystemNull/Private/OnlineSubsystemNull.cpp
@@ -108,6 +108,11 @@ IOnlinePresencePtr FOnlineSubsystemNull::GetPresenceInterface() const
 	return NULL;
 }
 
+IOnlineConnectionPtr FOnlineSubsystemNull::GetConnectionInterface() const
+{
+	return NULL;
+}
+
 bool FOnlineSubsystemNull::Tick(float DeltaTime)
 {
 	if (OnlineAsyncTaskThreadRunnable)
diff --git a/Engine/Source/Runtime/Online/OnlineSubsystemNull/Public/OnlineSubsystemNull.h b/Engine/Source/Runtime/Online/OnlineSubsystemNull/Public/OnlineSubsystemNull.h
--- a/Engine/Source/Runtime/Online/OnlineSubsystemNull/Public/OnlineSubsystemNull.h
+++ b/Engine/Source/Runtime/Online/OnlineSubsystemNull/Public/OnlineSubsystemNull.h
@@ -52,6 +52,7 @@ public:
 	virtual IOnlineUserPtr GetUserInterface() const override;
 	virtual IOnlineMessagePtr GetMessageInterface() const override;
 	virtual IOnlinePresencePtr GetPresenceInterface() const override;
+	virtual IOnlineConnectionPtr GetConnectionInterface() const override;
 
 	virtual bool Init() override;
 	virtual bool Shutdown() override;
diff --git a/Engine/Source/Runtime/Online/OnlineSubsystemSteam/Private/OnlineSubsystemSteam.cpp b/Engine/Source/Runtime/Online/OnlineSubsystemSteam/Private/OnlineSubsystemSteam.cpp
--- a/Engine/Source/Runtime/Online/OnlineSubsystemSteam/Private/OnlineSubsystemSteam.cpp
+++ b/Engine/Source/Runtime/Online/OnlineSubsystemSteam/Private/OnlineSubsystemSteam.cpp
@@ -249,6 +249,12 @@ IOnlinePresencePtr FOnlineSubsystemSteam::GetPresenceInterface() const
 	return NULL;
 }
 
+IOnlineConnectionPtr FOnlineSubsystemSteam::GetConnectionInterface() const
+{
+	return NULL;
+}
+
+
 void FOnlineSubsystemSteam::QueueAsyncTask(FOnlineAsyncTask* AsyncTask)
 {
 	check(OnlineAsyncTaskThreadRunnable);
diff --git a/Engine/Source/Runtime/Online/OnlineSubsystemSteam/Public/OnlineSubsystemSteam.h b/Engine/Source/Runtime/Online/OnlineSubsystemSteam/Public/OnlineSubsystemSteam.h
--- a/Engine/Source/Runtime/Online/OnlineSubsystemSteam/Public/OnlineSubsystemSteam.h
+++ b/Engine/Source/Runtime/Online/OnlineSubsystemSteam/Public/OnlineSubsystemSteam.h
@@ -230,6 +230,8 @@ public:
 	virtual IOnlineUserPtr GetUserInterface() const override;
 	virtual IOnlineMessagePtr GetMessageInterface() const override;
 	virtual IOnlinePresencePtr GetPresenceInterface() const override;
+	virtual IOnlineConnectionPtr GetConnectionInterface() const override;
+
 	virtual bool IsLocalPlayer(const FUniqueNetId& UniqueId) const override;
 	virtual bool Init() override;
 	virtual bool Shutdown() override;
diff --git a/Engine/Source/Runtime/Online/OnlineSubsystemUtils/Classes/ConnectionCallbackProxy.h b/Engine/Source/Runtime/Online/OnlineSubsystemUtils/Classes/ConnectionCallbackProxy.h
new file mode 100644
--- /dev/null
+++ b/Engine/Source/Runtime/Online/OnlineSubsystemUtils/Classes/ConnectionCallbackProxy.h
@@ -0,0 +1,39 @@
+// Copyright 1998-2014 Epic Games, Inc. All Rights Reserved.
+
+#pragma once
+
+#include "Runtime/Online/OnlineSubsystem/Public/Interfaces/OnlineConnectionInterface.h"
+#include "ConnectionCallbackProxy.generated.h"
+
+DECLARE_DYNAMIC_MULTICAST_DELEGATE_OneParam(FOnlineConnectionResult, int32, ErrorCode);
+
+
+UCLASS(MinimalAPI)
+class UConnectionCallbackProxy : public UOnlineBlueprintCallProxyBase
+{
+	GENERATED_UCLASS_BODY()
+
+	// Called when there is a successful query
+	UPROPERTY(BlueprintAssignable)
+	FOnlineConnectionResult OnSuccess;
+
+	// Called when there is an unsuccessful query
+	UPROPERTY(BlueprintAssignable)
+	FOnlineConnectionResult OnFailure;
+
+	// Connects to an online service such as Google Play
+	UFUNCTION(BlueprintCallable, meta = (BlueprintInternalUseOnly = "true"), Category = "Online|Achievements")
+		static UConnectionCallbackProxy* ConnectToService(class APlayerController* PlayerController);
+
+	// UOnlineBlueprintCallProxyBase interface
+	virtual void Activate() override;
+	// End of UOnlineBlueprintCallProxyBase interface
+
+private:
+	// Internal callback when the achievement query completes, calls out to the public success/failure callbacks
+	void OnConnectCompleted(int errorCode, bool bSuccess);
+
+private:
+	// The player controller triggering things
+	TWeakObjectPtr<APlayerController> PlayerControllerWeakPtr;
+};
\ No newline at end of file
diff --git a/Engine/Source/Runtime/Online/OnlineSubsystemUtils/Private/ConnectionCallbackProxy.cpp b/Engine/Source/Runtime/Online/OnlineSubsystemUtils/Private/ConnectionCallbackProxy.cpp
new file mode 100644
--- /dev/null
+++ b/Engine/Source/Runtime/Online/OnlineSubsystemUtils/Private/ConnectionCallbackProxy.cpp
@@ -0,0 +1,57 @@
+// Copyright 1998-2014 Epic Games, Inc. All Rights Reserved.
+
+#include "OnlineSubsystemUtilsPrivatePCH.h"
+
+//////////////////////////////////////////////////////////////////////////
+// UConnectionCallbackProxy
+
+UConnectionCallbackProxy::UConnectionCallbackProxy(const class FPostConstructInitializeProperties& PCIP)
+	: Super(PCIP)
+{
+}
+
+UConnectionCallbackProxy* UConnectionCallbackProxy::ConnectToService(class APlayerController* PlayerController)
+{
+	UConnectionCallbackProxy* Proxy = NewObject<UConnectionCallbackProxy>();
+	Proxy->PlayerControllerWeakPtr = PlayerController;
+	return Proxy;
+}
+
+void UConnectionCallbackProxy::Activate()
+{
+	FOnlineSubsystemBPCallHelper Helper(TEXT("ConnectToService"));
+	Helper.QueryIDFromPlayerController(PlayerControllerWeakPtr.Get());
+
+	if (Helper.IsValid())
+	{
+		IOnlineConnectionPtr Connection = Helper.OnlineSub->GetConnectionInterface();
+		if (Connection.IsValid())
+		{
+			FOnConnectionCompleteDelegate ConnectionFinishedDelegate = FOnConnectionCompleteDelegate::CreateUObject(this, &ThisClass::OnConnectCompleted);
+
+			Connection->Connect(ConnectionFinishedDelegate); // this will get called when everything is resolved
+
+			// OnQueryCompleted will get called, nothing more to do now
+			return;
+		}
+		else
+		{
+			FFrame::KismetExecutionMessage(TEXT("Connection control not supported by online subsystem"), ELogVerbosity::Warning);
+		}
+	}
+
+	// Fail immediately
+	OnFailure.Broadcast(0);
+}
+
+void UConnectionCallbackProxy::OnConnectCompleted(int errorCode, bool bSuccess)
+{
+	if (bSuccess)
+	{
+		OnSuccess.Broadcast(0);
+	}
+	else
+	{
+		OnFailure.Broadcast(0);
+	}
+}
