Fix for gnustl_shared not being loaded on Android devices before UE4 is loaded.

diff --git a/Engine/Build/Android/Java/src/com/epicgames/ue4/GameActivity.java b/Engine/Build/Android/Java/src/com/epicgames/ue4/GameActivity.java
--- a/Engine/Build/Android/Java/src/com/epicgames/ue4/GameActivity.java
+++ b/Engine/Build/Android/Java/src/com/epicgames/ue4/GameActivity.java
@@ -1067,6 +1067,7 @@ public class GameActivity extends NativeActivity implements GoogleApiClient.Conn
 
 	static
 	{
+		System.loadLibrary("gnustl_shared");
 		System.loadLibrary("UE4");
 	}
 }
