Removed Vibrator.hasVibrator() check since it was added in API11
#ue4
#Android

diff --git a/Engine/Build/Android/Java/src/com/epicgames/ue4/GameActivity.java b/Engine/Build/Android/Java/src/com/epicgames/ue4/GameActivity.java
--- a/Engine/Build/Android/Java/src/com/epicgames/ue4/GameActivity.java
+++ b/Engine/Build/Android/Java/src/com/epicgames/ue4/GameActivity.java
@@ -462,14 +462,11 @@ public class GameActivity extends NativeActivity
 	public void AndroidThunkJava_Vibrate(long Duration)
 	{
 		Vibrator vibrator = (Vibrator)getSystemService(VIBRATOR_SERVICE);
-		if (vibrator.hasVibrator())
+		if (Duration < 1)
 		{
-			if (Duration < 1)
-			{
-				vibrator.cancel();
-			} else {
-				vibrator.vibrate(Duration);
-			}
+			vibrator.cancel();
+		} else {
+			vibrator.vibrate(Duration);
 		}
 	}
 
