Add more ProGuard exceptions, for the Google Play Games C++ SDK
Fixes TTP 347366 and "for distribution" Android builds

diff --git a/Engine/Build/Android/Java/proguard-project.txt b/Engine/Build/Android/Java/proguard-project.txt
--- a/Engine/Build/Android/Java/proguard-project.txt
+++ b/Engine/Build/Android/Java/proguard-project.txt
@@ -47,6 +47,15 @@
     public static final ** CREATOR;
 }
 
+# These prevent required classes and fields from being stripped when using the Google Play Games C++ SDK
+-keep class com.google.android.gms.common.** {
+	*;
+}
+
+-keep class com.google.android.gms.games.** {
+	*;
+}
+
 # End Google Play services requirements
 
 # Google Play Admob
