Rollback //depot/UE4/Engine/Build/Android/Java/proguard-project.txt to revision 6

diff --git a/Engine/Build/Android/Java/proguard-project.txt b/Engine/Build/Android/Java/proguard-project.txt
new file mode 100644
--- /dev/null
+++ b/Engine/Build/Android/Java/proguard-project.txt
@@ -0,0 +1,68 @@
+# To enable ProGuard in your project, edit project.properties
+# to define the proguard.config property as described in that file.
+#
+# Add project specific ProGuard rules here.
+# By default, the flags in this file are appended to flags specified
+# in ${sdk.dir}/tools/proguard/proguard-android.txt
+# You can edit the include path and order by changing the ProGuard
+# include property in project.properties.
+#
+# For more details, see
+#   http://developer.android.com/guide/developing/tools/proguard.html
+
+# Add any project specific keep options here:
+
+# Needed to keep the JNI functions around
+-keepclassmembers class com.epicgames.ue4.GameActivity {
+   public *;
+}
+
+# Needed...
+-keep class com.epicgames.ue4.GameActivity$JavaAchievement {
+   public *;
+}
+
+# If your project uses WebView with JS, uncomment the following
+# and specify the fully qualified class name to the JavaScript interface
+# class:
+#-keepclassmembers class fqcn.of.javascript.interface.for.webview {
+#   public *;
+#}
+
+# Required for Google Play services: http://developer.android.com/google/play-services/setup.html
+-keep class * extends java.util.ListResourceBundle {
+    protected Object[][] getContents();
+}
+
+-keep public class com.google.android.gms.common.internal.safeparcel.SafeParcelable {
+    public static final *** NULL;
+}
+
+-keepnames @com.google.android.gms.common.annotation.KeepName class *
+-keepclassmembernames class * {
+    @com.google.android.gms.common.annotation.KeepName *;
+}
+
+-keepnames class * implements android.os.Parcelable {
+    public static final ** CREATOR;
+}
+
+# These prevent required classes and fields from being stripped when using the Google Play Games C++ SDK
+-keep class com.google.android.gms.common.** {
+	*;
+}
+
+-keep class com.google.android.gms.games.** {
+	*;
+}
+
+# End Google Play services requirements
+
+# Google Play Admob
+-keep public class com.google.android.gms.ads.** {
+   public *;
+}
+
+-keep public class com.google.ads.** {
+   public *;
+}
