In App Purchases Google Play.
- IAP added using the google play billing helper functionality
- App Licensing Key and checking added, but not yet used.

- Changes to Store purchases to allow consumable flags.
- New maps added for android testing, this uses static responses for IAP purchase testing.

A few improvements left
- need licensing working for the query process.
- need testing without static responses. Work with google test responses.

diff --git a/Engine/Build/Android/Java/AndroidManifest.xml b/Engine/Build/Android/Java/AndroidManifest.xml
--- a/Engine/Build/Android/Java/AndroidManifest.xml
+++ b/Engine/Build/Android/Java/AndroidManifest.xml
@@ -43,4 +43,5 @@
 	<uses-permission android:name="android.permission.MODIFY_AUDIO_SETTINGS"/>
 	<uses-permission android:name="android.permission.GET_ACCOUNTS" />
 	<uses-permission android:name="android.permission.VIBRATE"/>
+	<uses-permission android:name="com.android.vending.BILLING" />
 </manifest>
diff --git a/Engine/Build/Android/Java/proguard-project.txt b/Engine/Build/Android/Java/proguard-project.txt
--- a/Engine/Build/Android/Java/proguard-project.txt
+++ b/Engine/Build/Android/Java/proguard-project.txt
@@ -66,3 +66,16 @@
 -keep public class com.google.ads.** {
    public *;
 }
+
+# Google Play Billing Services for In-App Purchases
+-keep class com.android.vending.billing.** {
+	*;
+}
+
+-keep class epicgames.ue4.StoreHelper.** {
+	public *;
+}
+
+-keep class epicgames.ue4.GooglePlayStoreHelper.** {
+	public *;
+}
diff --git a/Engine/Build/Android/Java/src/com/android/vending/billing/IInAppBillingService.aidl b/Engine/Build/Android/Java/src/com/android/vending/billing/IInAppBillingService.aidl
new file mode 100644
--- /dev/null
+++ b/Engine/Build/Android/Java/src/com/android/vending/billing/IInAppBillingService.aidl
@@ -0,0 +1,144 @@
+/*
+ * Copyright (C) 2012 The Android Open Source Project
+ *
+ * Licensed under the Apache License, Version 2.0 (the "License");
+ * you may not use this file except in compliance with the License.
+ * You may obtain a copy of the License at
+ *
+ *      http://www.apache.org/licenses/LICENSE-2.0
+ *
+ * Unless required by applicable law or agreed to in writing, software
+ * distributed under the License is distributed on an "AS IS" BASIS,
+ * WITHOUT WARRANTIES OR CONDITIONS OF ANY KIND, either express or implied.
+ * See the License for the specific language governing permissions and
+ * limitations under the License.
+ */
+
+package com.android.vending.billing;
+
+import android.os.Bundle;
+
+/**
+ * InAppBillingService is the service that provides in-app billing version 3 and beyond.
+ * This service provides the following features:
+ * 1. Provides a new API to get details of in-app items published for the app including
+ *    price, type, title and description.
+ * 2. The purchase flow is synchronous and purchase information is available immediately
+ *    after it completes.
+ * 3. Purchase information of in-app purchases is maintained within the Google Play system
+ *    till the purchase is consumed.
+ * 4. An API to consume a purchase of an inapp item. All purchases of one-time
+ *    in-app items are consumable and thereafter can be purchased again.
+ * 5. An API to get current purchases of the user immediately. This will not contain any
+ *    consumed purchases.
+ *
+ * All calls will give a response code with the following possible values
+ * RESULT_OK = 0 - success
+ * RESULT_USER_CANCELED = 1 - user pressed back or canceled a dialog
+ * RESULT_BILLING_UNAVAILABLE = 3 - this billing API version is not supported for the type requested
+ * RESULT_ITEM_UNAVAILABLE = 4 - requested SKU is not available for purchase
+ * RESULT_DEVELOPER_ERROR = 5 - invalid arguments provided to the API
+ * RESULT_ERROR = 6 - Fatal error during the API action
+ * RESULT_ITEM_ALREADY_OWNED = 7 - Failure to purchase since item is already owned
+ * RESULT_ITEM_NOT_OWNED = 8 - Failure to consume since item is not owned
+ */
+interface IInAppBillingService {
+    /**
+     * Checks support for the requested billing API version, package and in-app type.
+     * Minimum API version supported by this interface is 3.
+     * @param apiVersion the billing version which the app is using
+     * @param packageName the package name of the calling app
+     * @param type type of the in-app item being purchased "inapp" for one-time purchases
+     *        and "subs" for subscription.
+     * @return RESULT_OK(0) on success, corresponding result code on failures
+     */
+    int isBillingSupported(int apiVersion, String packageName, String type);
+
+    /**
+     * Provides details of a list of SKUs
+     * Given a list of SKUs of a valid type in the skusBundle, this returns a bundle
+     * with a list JSON strings containing the productId, price, title and description.
+     * This API can be called with a maximum of 20 SKUs.
+     * @param apiVersion billing API version that the Third-party is using
+     * @param packageName the package name of the calling app
+     * @param skusBundle bundle containing a StringArrayList of SKUs with key "ITEM_ID_LIST"
+     * @return Bundle containing the following key-value pairs
+     *         "RESPONSE_CODE" with int value, RESULT_OK(0) if success, other response codes on
+     *              failure as listed above.
+     *         "DETAILS_LIST" with a StringArrayList containing purchase information
+     *              in JSON format similar to:
+     *              '{ "productId" : "exampleSku", "type" : "inapp", "price" : "$5.00",
+     *                 "title : "Example Title", "description" : "This is an example description" }'
+     */
+    Bundle getSkuDetails(int apiVersion, String packageName, String type, in Bundle skusBundle);
+
+    /**
+     * Returns a pending intent to launch the purchase flow for an in-app item by providing a SKU,
+     * the type, a unique purchase token and an optional developer payload.
+     * @param apiVersion billing API version that the app is using
+     * @param packageName package name of the calling app
+     * @param sku the SKU of the in-app item as published in the developer console
+     * @param type the type of the in-app item ("inapp" for one-time purchases
+     *        and "subs" for subscription).
+     * @param developerPayload optional argument to be sent back with the purchase information
+     * @return Bundle containing the following key-value pairs
+     *         "RESPONSE_CODE" with int value, RESULT_OK(0) if success, other response codes on
+     *              failure as listed above.
+     *         "BUY_INTENT" - PendingIntent to start the purchase flow
+     *
+     * The Pending intent should be launched with startIntentSenderForResult. When purchase flow
+     * has completed, the onActivityResult() will give a resultCode of OK or CANCELED.
+     * If the purchase is successful, the result data will contain the following key-value pairs
+     *         "RESPONSE_CODE" with int value, RESULT_OK(0) if success, other response codes on
+     *              failure as listed above.
+     *         "INAPP_PURCHASE_DATA" - String in JSON format similar to
+     *              '{"orderId":"12999763169054705758.1371079406387615",
+     *                "packageName":"com.example.app",
+     *                "productId":"exampleSku",
+     *                "purchaseTime":1345678900000,
+     *                "purchaseToken" : "122333444455555",
+     *                "developerPayload":"example developer payload" }'
+     *         "INAPP_DATA_SIGNATURE" - String containing the signature of the purchase data that
+     *                                  was signed with the private key of the developer
+     *                                  TODO: change this to app-specific keys.
+     */
+    Bundle getBuyIntent(int apiVersion, String packageName, String sku, String type,
+        String developerPayload);
+
+    /**
+     * Returns the current SKUs owned by the user of the type and package name specified along with
+     * purchase information and a signature of the data to be validated.
+     * This will return all SKUs that have been purchased in V3 and managed items purchased using
+     * V1 and V2 that have not been consumed.
+     * @param apiVersion billing API version that the app is using
+     * @param packageName package name of the calling app
+     * @param type the type of the in-app items being requested
+     *        ("inapp" for one-time purchases and "subs" for subscription).
+     * @param continuationToken to be set as null for the first call, if the number of owned
+     *        skus are too many, a continuationToken is returned in the response bundle.
+     *        This method can be called again with the continuation token to get the next set of
+     *        owned skus.
+     * @return Bundle containing the following key-value pairs
+     *         "RESPONSE_CODE" with int value, RESULT_OK(0) if success, other response codes on
+     *              failure as listed above.
+     *         "INAPP_PURCHASE_ITEM_LIST" - StringArrayList containing the list of SKUs
+     *         "INAPP_PURCHASE_DATA_LIST" - StringArrayList containing the purchase information
+     *         "INAPP_DATA_SIGNATURE_LIST"- StringArrayList containing the signatures
+     *                                      of the purchase information
+     *         "INAPP_CONTINUATION_TOKEN" - String containing a continuation token for the
+     *                                      next set of in-app purchases. Only set if the
+     *                                      user has more owned skus than the current list.
+     */
+    Bundle getPurchases(int apiVersion, String packageName, String type, String continuationToken);
+
+    /**
+     * Consume the last purchase of the given SKU. This will result in this item being removed
+     * from all subsequent responses to getPurchases() and allow re-purchase of this item.
+     * @param apiVersion billing API version that the app is using
+     * @param packageName package name of the calling app
+     * @param purchaseToken token in the purchase information JSON that identifies the purchase
+     *        to be consumed
+     * @return 0 if consumption succeeded. Appropriate error values for failures.
+     */
+    int consumePurchase(int apiVersion, String packageName, String purchaseToken);
+}
diff --git a/Engine/Build/Android/Java/src/com/android/vending/billing/util/Base64.java b/Engine/Build/Android/Java/src/com/android/vending/billing/util/Base64.java
new file mode 100644
--- /dev/null
+++ b/Engine/Build/Android/Java/src/com/android/vending/billing/util/Base64.java
@@ -0,0 +1,570 @@
+// Portions copyright 2002, Google, Inc.
+//
+// Licensed under the Apache License, Version 2.0 (the "License");
+// you may not use this file except in compliance with the License.
+// You may obtain a copy of the License at
+//
+//     http://www.apache.org/licenses/LICENSE-2.0
+//
+// Unless required by applicable law or agreed to in writing, software
+// distributed under the License is distributed on an "AS IS" BASIS,
+// WITHOUT WARRANTIES OR CONDITIONS OF ANY KIND, either express or implied.
+// See the License for the specific language governing permissions and
+// limitations under the License.
+
+package com.android.vending.billing.util;
+
+// This code was converted from code at http://iharder.sourceforge.net/base64/
+// Lots of extraneous features were removed.
+/* The original code said:
+ * <p>
+ * I am placing this code in the Public Domain. Do with it as you will.
+ * This software comes with no guarantees or warranties but with
+ * plenty of well-wishing instead!
+ * Please visit
+ * <a href="http://iharder.net/xmlizable">http://iharder.net/xmlizable</a>
+ * periodically to check for updates or to contribute improvements.
+ * </p>
+ *
+ * @author Robert Harder
+ * @author [REDACTED_EMAIL]
+ * @version 1.3
+ */
+
+/**
+ * Base64 converter class. This code is not a complete MIME encoder;
+ * it simply converts binary data to base64 data and back.
+ *
+ * <p>Note {@link CharBase64} is a GWT-compatible implementation of this
+ * class.
+ */
+public class Base64 {
+    /** Specify encoding (value is {@code true}). */
+    public final static boolean ENCODE = true;
+
+    /** Specify decoding (value is {@code false}). */
+    public final static boolean DECODE = false;
+
+    /** The equals sign (=) as a byte. */
+    private final static byte EQUALS_SIGN = (byte) '=';
+
+    /** The new line character (\n) as a byte. */
+    private final static byte NEW_LINE = (byte) '\n';
+
+    /**
+     * The 64 valid Base64 values.
+     */
+    private final static byte[] ALPHABET =
+        {(byte) 'A', (byte) 'B', (byte) 'C', (byte) 'D', (byte) 'E', (byte) 'F',
+        (byte) 'G', (byte) 'H', (byte) 'I', (byte) 'J', (byte) 'K',
+        (byte) 'L', (byte) 'M', (byte) 'N', (byte) 'O', (byte) 'P',
+        (byte) 'Q', (byte) 'R', (byte) 'S', (byte) 'T', (byte) 'U',
+        (byte) 'V', (byte) 'W', (byte) 'X', (byte) 'Y', (byte) 'Z',
+        (byte) 'a', (byte) 'b', (byte) 'c', (byte) 'd', (byte) 'e',
+        (byte) 'f', (byte) 'g', (byte) 'h', (byte) 'i', (byte) 'j',
+        (byte) 'k', (byte) 'l', (byte) 'm', (byte) 'n', (byte) 'o',
+        (byte) 'p', (byte) 'q', (byte) 'r', (byte) 's', (byte) 't',
+        (byte) 'u', (byte) 'v', (byte) 'w', (byte) 'x', (byte) 'y',
+        (byte) 'z', (byte) '0', (byte) '1', (byte) '2', (byte) '3',
+        (byte) '4', (byte) '5', (byte) '6', (byte) '7', (byte) '8',
+        (byte) '9', (byte) '+', (byte) '/'};
+
+    /**
+     * The 64 valid web safe Base64 values.
+     */
+    private final static byte[] WEBSAFE_ALPHABET =
+        {(byte) 'A', (byte) 'B', (byte) 'C', (byte) 'D', (byte) 'E', (byte) 'F',
+        (byte) 'G', (byte) 'H', (byte) 'I', (byte) 'J', (byte) 'K',
+        (byte) 'L', (byte) 'M', (byte) 'N', (byte) 'O', (byte) 'P',
+        (byte) 'Q', (byte) 'R', (byte) 'S', (byte) 'T', (byte) 'U',
+        (byte) 'V', (byte) 'W', (byte) 'X', (byte) 'Y', (byte) 'Z',
+        (byte) 'a', (byte) 'b', (byte) 'c', (byte) 'd', (byte) 'e',
+        (byte) 'f', (byte) 'g', (byte) 'h', (byte) 'i', (byte) 'j',
+        (byte) 'k', (byte) 'l', (byte) 'm', (byte) 'n', (byte) 'o',
+        (byte) 'p', (byte) 'q', (byte) 'r', (byte) 's', (byte) 't',
+        (byte) 'u', (byte) 'v', (byte) 'w', (byte) 'x', (byte) 'y',
+        (byte) 'z', (byte) '0', (byte) '1', (byte) '2', (byte) '3',
+        (byte) '4', (byte) '5', (byte) '6', (byte) '7', (byte) '8',
+        (byte) '9', (byte) '-', (byte) '_'};
+
+    /**
+     * Translates a Base64 value to either its 6-bit reconstruction value
+     * or a negative number indicating some other meaning.
+     **/
+    private final static byte[] DECODABET = {-9, -9, -9, -9, -9, -9, -9, -9, -9, // Decimal  0 -  8
+        -5, -5, // Whitespace: Tab and Linefeed
+        -9, -9, // Decimal 11 - 12
+        -5, // Whitespace: Carriage Return
+        -9, -9, -9, -9, -9, -9, -9, -9, -9, -9, -9, -9, -9, // Decimal 14 - 26
+        -9, -9, -9, -9, -9, // Decimal 27 - 31
+        -5, // Whitespace: Space
+        -9, -9, -9, -9, -9, -9, -9, -9, -9, -9, // Decimal 33 - 42
+        62, // Plus sign at decimal 43
+        -9, -9, -9, // Decimal 44 - 46
+        63, // Slash at decimal 47
+        52, 53, 54, 55, 56, 57, 58, 59, 60, 61, // Numbers zero through nine
+        -9, -9, -9, // Decimal 58 - 60
+        -1, // Equals sign at decimal 61
+        -9, -9, -9, // Decimal 62 - 64
+        0, 1, 2, 3, 4, 5, 6, 7, 8, 9, 10, 11, 12, 13, // Letters 'A' through 'N'
+        14, 15, 16, 17, 18, 19, 20, 21, 22, 23, 24, 25, // Letters 'O' through 'Z'
+        -9, -9, -9, -9, -9, -9, // Decimal 91 - 96
+        26, 27, 28, 29, 30, 31, 32, 33, 34, 35, 36, 37, 38, // Letters 'a' through 'm'
+        39, 40, 41, 42, 43, 44, 45, 46, 47, 48, 49, 50, 51, // Letters 'n' through 'z'
+        -9, -9, -9, -9, -9 // Decimal 123 - 127
+        /*  ,-9,-9,-9,-9,-9,-9,-9,-9,-9,-9,-9,-9,     // Decimal 128 - 139
+        -9,-9,-9,-9,-9,-9,-9,-9,-9,-9,-9,-9,-9,     // Decimal 140 - 152
+        -9,-9,-9,-9,-9,-9,-9,-9,-9,-9,-9,-9,-9,     // Decimal 153 - 165
+        -9,-9,-9,-9,-9,-9,-9,-9,-9,-9,-9,-9,-9,     // Decimal 166 - 178
+        -9,-9,-9,-9,-9,-9,-9,-9,-9,-9,-9,-9,-9,     // Decimal 179 - 191
+        -9,-9,-9,-9,-9,-9,-9,-9,-9,-9,-9,-9,-9,     // Decimal 192 - 204
+        -9,-9,-9,-9,-9,-9,-9,-9,-9,-9,-9,-9,-9,     // Decimal 205 - 217
+        -9,-9,-9,-9,-9,-9,-9,-9,-9,-9,-9,-9,-9,     // Decimal 218 - 230
+        -9,-9,-9,-9,-9,-9,-9,-9,-9,-9,-9,-9,-9,     // Decimal 231 - 243
+        -9,-9,-9,-9,-9,-9,-9,-9,-9,-9,-9,-9         // Decimal 244 - 255 */
+    };
+
+    /** The web safe decodabet */
+    private final static byte[] WEBSAFE_DECODABET =
+        {-9, -9, -9, -9, -9, -9, -9, -9, -9, // Decimal  0 -  8
+        -5, -5, // Whitespace: Tab and Linefeed
+        -9, -9, // Decimal 11 - 12
+        -5, // Whitespace: Carriage Return
+        -9, -9, -9, -9, -9, -9, -9, -9, -9, -9, -9, -9, -9, // Decimal 14 - 26
+        -9, -9, -9, -9, -9, // Decimal 27 - 31
+        -5, // Whitespace: Space
+        -9, -9, -9, -9, -9, -9, -9, -9, -9, -9, -9, -9, // Decimal 33 - 44
+        62, // Dash '-' sign at decimal 45
+        -9, -9, // Decimal 46-47
+        52, 53, 54, 55, 56, 57, 58, 59, 60, 61, // Numbers zero through nine
+        -9, -9, -9, // Decimal 58 - 60
+        -1, // Equals sign at decimal 61
+        -9, -9, -9, // Decimal 62 - 64
+        0, 1, 2, 3, 4, 5, 6, 7, 8, 9, 10, 11, 12, 13, // Letters 'A' through 'N'
+        14, 15, 16, 17, 18, 19, 20, 21, 22, 23, 24, 25, // Letters 'O' through 'Z'
+        -9, -9, -9, -9, // Decimal 91-94
+        63, // Underscore '_' at decimal 95
+        -9, // Decimal 96
+        26, 27, 28, 29, 30, 31, 32, 33, 34, 35, 36, 37, 38, // Letters 'a' through 'm'
+        39, 40, 41, 42, 43, 44, 45, 46, 47, 48, 49, 50, 51, // Letters 'n' through 'z'
+        -9, -9, -9, -9, -9 // Decimal 123 - 127
+        /*  ,-9,-9,-9,-9,-9,-9,-9,-9,-9,-9,-9,-9,     // Decimal 128 - 139
+        -9,-9,-9,-9,-9,-9,-9,-9,-9,-9,-9,-9,-9,     // Decimal 140 - 152
+        -9,-9,-9,-9,-9,-9,-9,-9,-9,-9,-9,-9,-9,     // Decimal 153 - 165
+        -9,-9,-9,-9,-9,-9,-9,-9,-9,-9,-9,-9,-9,     // Decimal 166 - 178
+        -9,-9,-9,-9,-9,-9,-9,-9,-9,-9,-9,-9,-9,     // Decimal 179 - 191
+        -9,-9,-9,-9,-9,-9,-9,-9,-9,-9,-9,-9,-9,     // Decimal 192 - 204
+        -9,-9,-9,-9,-9,-9,-9,-9,-9,-9,-9,-9,-9,     // Decimal 205 - 217
+        -9,-9,-9,-9,-9,-9,-9,-9,-9,-9,-9,-9,-9,     // Decimal 218 - 230
+        -9,-9,-9,-9,-9,-9,-9,-9,-9,-9,-9,-9,-9,     // Decimal 231 - 243
+        -9,-9,-9,-9,-9,-9,-9,-9,-9,-9,-9,-9         // Decimal 244 - 255 */
+        };
+
+    // Indicates white space in encoding
+    private final static byte WHITE_SPACE_ENC = -5;
+    // Indicates equals sign in encoding
+    private final static byte EQUALS_SIGN_ENC = -1;
+
+    /** Defeats instantiation. */
+    private Base64() {
+    }
+
+    /* ********  E N C O D I N G   M E T H O D S  ******** */
+
+    /**
+     * Encodes up to three bytes of the array <var>source</var>
+     * and writes the resulting four Base64 bytes to <var>destination</var>.
+     * The source and destination arrays can be manipulated
+     * anywhere along their length by specifying
+     * <var>srcOffset</var> and <var>destOffset</var>.
+     * This method does not check to make sure your arrays
+     * are large enough to accommodate <var>srcOffset</var> + 3 for
+     * the <var>source</var> array or <var>destOffset</var> + 4 for
+     * the <var>destination</var> array.
+     * The actual number of significant bytes in your array is
+     * given by <var>numSigBytes</var>.
+     *
+     * @param source the array to convert
+     * @param srcOffset the index where conversion begins
+     * @param numSigBytes the number of significant bytes in your array
+     * @param destination the array to hold the conversion
+     * @param destOffset the index where output will be put
+     * @param alphabet is the encoding alphabet
+     * @return the <var>destination</var> array
+     * @since 1.3
+     */
+    private static byte[] encode3to4(byte[] source, int srcOffset,
+            int numSigBytes, byte[] destination, int destOffset, byte[] alphabet) {
+        //           1         2         3
+        // 01234567890123456789012345678901 Bit position
+        // --------000000001111111122222222 Array position from threeBytes
+        // --------|    ||    ||    ||    | Six bit groups to index alphabet
+        //          >>18  >>12  >> 6  >> 0  Right shift necessary
+        //                0x3f  0x3f  0x3f  Additional AND
+
+        // Create buffer with zero-padding if there are only one or two
+        // significant bytes passed in the array.
+        // We have to shift left 24 in order to flush out the 1's that appear
+        // when Java treats a value as negative that is cast from a byte to an int.
+        int inBuff =
+                (numSigBytes > 0 ? ((source[srcOffset] << 24) >>> 8) : 0)
+                | (numSigBytes > 1 ? ((source[srcOffset + 1] << 24) >>> 16) : 0)
+                | (numSigBytes > 2 ? ((source[srcOffset + 2] << 24) >>> 24) : 0);
+
+        switch (numSigBytes) {
+            case 3:
+                destination[destOffset] = alphabet[(inBuff >>> 18)];
+                destination[destOffset + 1] = alphabet[(inBuff >>> 12) & 0x3f];
+                destination[destOffset + 2] = alphabet[(inBuff >>> 6) & 0x3f];
+                destination[destOffset + 3] = alphabet[(inBuff) & 0x3f];
+                return destination;
+            case 2:
+                destination[destOffset] = alphabet[(inBuff >>> 18)];
+                destination[destOffset + 1] = alphabet[(inBuff >>> 12) & 0x3f];
+                destination[destOffset + 2] = alphabet[(inBuff >>> 6) & 0x3f];
+                destination[destOffset + 3] = EQUALS_SIGN;
+                return destination;
+            case 1:
+                destination[destOffset] = alphabet[(inBuff >>> 18)];
+                destination[destOffset + 1] = alphabet[(inBuff >>> 12) & 0x3f];
+                destination[destOffset + 2] = EQUALS_SIGN;
+                destination[destOffset + 3] = EQUALS_SIGN;
+                return destination;
+            default:
+                return destination;
+        } // end switch
+    } // end encode3to4
+
+    /**
+     * Encodes a byte array into Base64 notation.
+     * Equivalent to calling
+     * {@code encodeBytes(source, 0, source.length)}
+     *
+     * @param source The data to convert
+     * @since 1.4
+     */
+    public static String encode(byte[] source) {
+        return encode(source, 0, source.length, ALPHABET, true);
+    }
+
+    /**
+     * Encodes a byte array into web safe Base64 notation.
+     *
+     * @param source The data to convert
+     * @param doPadding is {@code true} to pad result with '=' chars
+     *        if it does not fall on 3 byte boundaries
+     */
+    public static String encodeWebSafe(byte[] source, boolean doPadding) {
+        return encode(source, 0, source.length, WEBSAFE_ALPHABET, doPadding);
+    }
+
+    /**
+     * Encodes a byte array into Base64 notation.
+     *
+     * @param source the data to convert
+     * @param off offset in array where conversion should begin
+     * @param len length of data to convert
+     * @param alphabet the encoding alphabet
+     * @param doPadding is {@code true} to pad result with '=' chars
+     * if it does not fall on 3 byte boundaries
+     * @since 1.4
+     */
+    public static String encode(byte[] source, int off, int len, byte[] alphabet,
+            boolean doPadding) {
+        byte[] outBuff = encode(source, off, len, alphabet, Integer.MAX_VALUE);
+        int outLen = outBuff.length;
+
+        // If doPadding is false, set length to truncate '='
+        // padding characters
+        while (doPadding == false && outLen > 0) {
+            if (outBuff[outLen - 1] != '=') {
+                break;
+            }
+            outLen -= 1;
+        }
+
+        return new String(outBuff, 0, outLen);
+    }
+
+    /**
+     * Encodes a byte array into Base64 notation.
+     *
+     * @param source the data to convert
+     * @param off offset in array where conversion should begin
+     * @param len length of data to convert
+     * @param alphabet is the encoding alphabet
+     * @param maxLineLength maximum length of one line.
+     * @return the BASE64-encoded byte array
+     */
+    public static byte[] encode(byte[] source, int off, int len, byte[] alphabet,
+            int maxLineLength) {
+        int lenDiv3 = (len + 2) / 3; // ceil(len / 3)
+        int len43 = lenDiv3 * 4;
+        byte[] outBuff = new byte[len43 // Main 4:3
+                                  + (len43 / maxLineLength)]; // New lines
+
+        int d = 0;
+        int e = 0;
+        int len2 = len - 2;
+        int lineLength = 0;
+        for (; d < len2; d += 3, e += 4) {
+
+            // The following block of code is the same as
+            // encode3to4( source, d + off, 3, outBuff, e, alphabet );
+            // but inlined for faster encoding (~20% improvement)
+            int inBuff =
+                    ((source[d + off] << 24) >>> 8)
+                    | ((source[d + 1 + off] << 24) >>> 16)
+                    | ((source[d + 2 + off] << 24) >>> 24);
+            outBuff[e] = alphabet[(inBuff >>> 18)];
+            outBuff[e + 1] = alphabet[(inBuff >>> 12) & 0x3f];
+            outBuff[e + 2] = alphabet[(inBuff >>> 6) & 0x3f];
+            outBuff[e + 3] = alphabet[(inBuff) & 0x3f];
+
+            lineLength += 4;
+            if (lineLength == maxLineLength) {
+                outBuff[e + 4] = NEW_LINE;
+                e++;
+                lineLength = 0;
+            } // end if: end of line
+        } // end for: each piece of array
+
+        if (d < len) {
+            encode3to4(source, d + off, len - d, outBuff, e, alphabet);
+
+            lineLength += 4;
+            if (lineLength == maxLineLength) {
+                // Add a last newline
+                outBuff[e + 4] = NEW_LINE;
+                e++;
+            }
+            e += 4;
+        }
+
+        assert (e == outBuff.length);
+        return outBuff;
+    }
+
+
+    /* ********  D E C O D I N G   M E T H O D S  ******** */
+
+
+    /**
+     * Decodes four bytes from array <var>source</var>
+     * and writes the resulting bytes (up to three of them)
+     * to <var>destination</var>.
+     * The source and destination arrays can be manipulated
+     * anywhere along their length by specifying
+     * <var>srcOffset</var> and <var>destOffset</var>.
+     * This method does not check to make sure your arrays
+     * are large enough to accommodate <var>srcOffset</var> + 4 for
+     * the <var>source</var> array or <var>destOffset</var> + 3 for
+     * the <var>destination</var> array.
+     * This method returns the actual number of bytes that
+     * were converted from the Base64 encoding.
+     *
+     *
+     * @param source the array to convert
+     * @param srcOffset the index where conversion begins
+     * @param destination the array to hold the conversion
+     * @param destOffset the index where output will be put
+     * @param decodabet the decodabet for decoding Base64 content
+     * @return the number of decoded bytes converted
+     * @since 1.3
+     */
+    private static int decode4to3(byte[] source, int srcOffset,
+            byte[] destination, int destOffset, byte[] decodabet) {
+        // Example: Dk==
+        if (source[srcOffset + 2] == EQUALS_SIGN) {
+            int outBuff =
+                    ((decodabet[source[srcOffset]] << 24) >>> 6)
+                    | ((decodabet[source[srcOffset + 1]] << 24) >>> 12);
+
+            destination[destOffset] = (byte) (outBuff >>> 16);
+            return 1;
+        } else if (source[srcOffset + 3] == EQUALS_SIGN) {
+            // Example: DkL=
+            int outBuff =
+                    ((decodabet[source[srcOffset]] << 24) >>> 6)
+                    | ((decodabet[source[srcOffset + 1]] << 24) >>> 12)
+                    | ((decodabet[source[srcOffset + 2]] << 24) >>> 18);
+
+            destination[destOffset] = (byte) (outBuff >>> 16);
+            destination[destOffset + 1] = (byte) (outBuff >>> 8);
+            return 2;
+        } else {
+            // Example: DkLE
+            int outBuff =
+                    ((decodabet[source[srcOffset]] << 24) >>> 6)
+                    | ((decodabet[source[srcOffset + 1]] << 24) >>> 12)
+                    | ((decodabet[source[srcOffset + 2]] << 24) >>> 18)
+                    | ((decodabet[source[srcOffset + 3]] << 24) >>> 24);
+
+            destination[destOffset] = (byte) (outBuff >> 16);
+            destination[destOffset + 1] = (byte) (outBuff >> 8);
+            destination[destOffset + 2] = (byte) (outBuff);
+            return 3;
+        }
+    } // end decodeToBytes
+
+
+    /**
+     * Decodes data from Base64 notation.
+     *
+     * @param s the string to decode (decoded in default encoding)
+     * @return the decoded data
+     * @since 1.4
+     */
+    public static byte[] decode(String s) throws Base64DecoderException {
+        byte[] bytes = s.getBytes();
+        return decode(bytes, 0, bytes.length);
+    }
+
+    /**
+     * Decodes data from web safe Base64 notation.
+     * Web safe encoding uses '-' instead of '+', '_' instead of '/'
+     *
+     * @param s the string to decode (decoded in default encoding)
+     * @return the decoded data
+     */
+    public static byte[] decodeWebSafe(String s) throws Base64DecoderException {
+        byte[] bytes = s.getBytes();
+        return decodeWebSafe(bytes, 0, bytes.length);
+    }
+
+    /**
+     * Decodes Base64 content in byte array format and returns
+     * the decoded byte array.
+     *
+     * @param source The Base64 encoded data
+     * @return decoded data
+     * @since 1.3
+     * @throws Base64DecoderException
+     */
+    public static byte[] decode(byte[] source) throws Base64DecoderException {
+        return decode(source, 0, source.length);
+    }
+
+    /**
+     * Decodes web safe Base64 content in byte array format and returns
+     * the decoded data.
+     * Web safe encoding uses '-' instead of '+', '_' instead of '/'
+     *
+     * @param source the string to decode (decoded in default encoding)
+     * @return the decoded data
+     */
+    public static byte[] decodeWebSafe(byte[] source)
+            throws Base64DecoderException {
+        return decodeWebSafe(source, 0, source.length);
+    }
+
+    /**
+     * Decodes Base64 content in byte array format and returns
+     * the decoded byte array.
+     *
+     * @param source the Base64 encoded data
+     * @param off    the offset of where to begin decoding
+     * @param len    the length of characters to decode
+     * @return decoded data
+     * @since 1.3
+     * @throws Base64DecoderException
+     */
+    public static byte[] decode(byte[] source, int off, int len)
+            throws Base64DecoderException {
+        return decode(source, off, len, DECODABET);
+    }
+
+    /**
+     * Decodes web safe Base64 content in byte array format and returns
+     * the decoded byte array.
+     * Web safe encoding uses '-' instead of '+', '_' instead of '/'
+     *
+     * @param source the Base64 encoded data
+     * @param off    the offset of where to begin decoding
+     * @param len    the length of characters to decode
+     * @return decoded data
+     */
+    public static byte[] decodeWebSafe(byte[] source, int off, int len)
+            throws Base64DecoderException {
+        return decode(source, off, len, WEBSAFE_DECODABET);
+    }
+
+    /**
+     * Decodes Base64 content using the supplied decodabet and returns
+     * the decoded byte array.
+     *
+     * @param source the Base64 encoded data
+     * @param off the offset of where to begin decoding
+     * @param len the length of characters to decode
+     * @param decodabet the decodabet for decoding Base64 content
+     * @return decoded data
+     */
+    public static byte[] decode(byte[] source, int off, int len, byte[] decodabet)
+            throws Base64DecoderException {
+        int len34 = len * 3 / 4;
+        byte[] outBuff = new byte[2 + len34]; // Upper limit on size of output
+        int outBuffPosn = 0;
+
+        byte[] b4 = new byte[4];
+        int b4Posn = 0;
+        int i = 0;
+        byte sbiCrop = 0;
+        byte sbiDecode = 0;
+        for (i = 0; i < len; i++) {
+            sbiCrop = (byte) (source[i + off] & 0x7f); // Only the low seven bits
+            sbiDecode = decodabet[sbiCrop];
+
+            if (sbiDecode >= WHITE_SPACE_ENC) { // White space Equals sign or better
+                if (sbiDecode >= EQUALS_SIGN_ENC) {
+                    // An equals sign (for padding) must not occur at position 0 or 1
+                    // and must be the last byte[s] in the encoded value
+                    if (sbiCrop == EQUALS_SIGN) {
+                        int bytesLeft = len - i;
+                        byte lastByte = (byte) (source[len - 1 + off] & 0x7f);
+                        if (b4Posn == 0 || b4Posn == 1) {
+                            throw new Base64DecoderException(
+                                    "invalid padding byte '=' at byte offset " + i);
+                        } else if ((b4Posn == 3 && bytesLeft > 2)
+                                || (b4Posn == 4 && bytesLeft > 1)) {
+                            throw new Base64DecoderException(
+                                    "padding byte '=' falsely signals end of encoded value "
+                                            + "at offset " + i);
+                        } else if (lastByte != EQUALS_SIGN && lastByte != NEW_LINE) {
+                            throw new Base64DecoderException(
+                                    "encoded value has invalid trailing byte");
+                        }
+                        break;
+                    }
+
+                    b4[b4Posn++] = sbiCrop;
+                    if (b4Posn == 4) {
+                        outBuffPosn += decode4to3(b4, 0, outBuff, outBuffPosn, decodabet);
+                        b4Posn = 0;
+                    }
+                }
+            } else {
+                throw new Base64DecoderException("Bad Base64 input character at " + i
+                        + ": " + source[i + off] + "(decimal)");
+            }
+        }
+
+        // Because web safe encoding allows non padding base64 encodes, we
+        // need to pad the rest of the b4 buffer with equal signs when
+        // b4Posn != 0.  There can be at most 2 equal signs at the end of
+        // four characters, so the b4 buffer must have two or three
+        // characters.  This also catches the case where the input is
+        // padded with EQUALS_SIGN
+        if (b4Posn != 0) {
+            if (b4Posn == 1) {
+                throw new Base64DecoderException("single trailing character at offset "
+                        + (len - 1));
+            }
+            b4[b4Posn++] = EQUALS_SIGN;
+            outBuffPosn += decode4to3(b4, 0, outBuff, outBuffPosn, decodabet);
+        }
+
+        byte[] out = new byte[outBuffPosn];
+        System.arraycopy(outBuff, 0, out, 0, outBuffPosn);
+        return out;
+    }
+}
diff --git a/Engine/Build/Android/Java/src/com/android/vending/billing/util/Base64DecoderException.java b/Engine/Build/Android/Java/src/com/android/vending/billing/util/Base64DecoderException.java
new file mode 100644
--- /dev/null
+++ b/Engine/Build/Android/Java/src/com/android/vending/billing/util/Base64DecoderException.java
@@ -0,0 +1,32 @@
+// Copyright 2002, Google, Inc.
+//
+// Licensed under the Apache License, Version 2.0 (the "License");
+// you may not use this file except in compliance with the License.
+// You may obtain a copy of the License at
+//
+//     http://www.apache.org/licenses/LICENSE-2.0
+//
+// Unless required by applicable law or agreed to in writing, software
+// distributed under the License is distributed on an "AS IS" BASIS,
+// WITHOUT WARRANTIES OR CONDITIONS OF ANY KIND, either express or implied.
+// See the License for the specific language governing permissions and
+// limitations under the License.
+
+package com.android.vending.billing.util;
+
+/**
+ * Exception thrown when encountering an invalid Base64 input character.
+ *
+ * @author nelson
+ */
+public class Base64DecoderException extends Exception {
+    public Base64DecoderException() {
+        super();
+    }
+
+    public Base64DecoderException(String s) {
+        super(s);
+    }
+
+    private static final long serialVersionUID = 1L;
+}
diff --git a/Engine/Build/Android/Java/src/com/android/vending/billing/util/IabException.java b/Engine/Build/Android/Java/src/com/android/vending/billing/util/IabException.java
new file mode 100644
--- /dev/null
+++ b/Engine/Build/Android/Java/src/com/android/vending/billing/util/IabException.java
@@ -0,0 +1,43 @@
+/* Copyright (c) 2012 Google Inc.
+ *
+ * Licensed under the Apache License, Version 2.0 (the "License");
+ * you may not use this file except in compliance with the License.
+ * You may obtain a copy of the License at
+ *
+ *     http://www.apache.org/licenses/LICENSE-2.0
+ *
+ * Unless required by applicable law or agreed to in writing, software
+ * distributed under the License is distributed on an "AS IS" BASIS,
+ * WITHOUT WARRANTIES OR CONDITIONS OF ANY KIND, either express or implied.
+ * See the License for the specific language governing permissions and
+ * limitations under the License.
+ */
+ 
+package com.android.vending.billing.util;
+
+/**
+ * Exception thrown when something went wrong with in-app billing.
+ * An IabException has an associated IabResult (an error).
+ * To get the IAB result that caused this exception to be thrown,
+ * call {@link #getResult()}.
+ */
+public class IabException extends Exception {
+    IabResult mResult;
+
+    public IabException(IabResult r) {
+        this(r, null);
+    }
+    public IabException(int response, String message) {
+        this(new IabResult(response, message));
+    }
+    public IabException(IabResult r, Exception cause) {
+        super(r.getMessage(), cause);
+        mResult = r;
+    }
+    public IabException(int response, String message, Exception cause) {
+        this(new IabResult(response, message), cause);
+    }
+
+    /** Returns the IAB result (error) that this exception signals. */
+    public IabResult getResult() { return mResult; }
+}
\ No newline at end of file
diff --git a/Engine/Build/Android/Java/src/com/android/vending/billing/util/IabHelper.java b/Engine/Build/Android/Java/src/com/android/vending/billing/util/IabHelper.java
new file mode 100644
--- /dev/null
+++ b/Engine/Build/Android/Java/src/com/android/vending/billing/util/IabHelper.java
@@ -0,0 +1,991 @@
+/* Copyright (c) 2012 Google Inc.
+ *
+ * Licensed under the Apache License, Version 2.0 (the "License");
+ * you may not use this file except in compliance with the License.
+ * You may obtain a copy of the License at
+ *
+ *     http://www.apache.org/licenses/LICENSE-2.0
+ *
+ * Unless required by applicable law or agreed to in writing, software
+ * distributed under the License is distributed on an "AS IS" BASIS,
+ * WITHOUT WARRANTIES OR CONDITIONS OF ANY KIND, either express or implied.
+ * See the License for the specific language governing permissions and
+ * limitations under the License.
+ */
+ 
+package com.android.vending.billing.util;
+
+import android.app.Activity;
+import android.app.PendingIntent;
+import android.content.ComponentName;
+import android.content.Context;
+import android.content.Intent;
+import android.content.IntentSender.SendIntentException;
+import android.content.ServiceConnection;
+import android.os.Bundle;
+import android.os.Handler;
+import android.os.IBinder;
+import android.os.RemoteException;
+import android.text.TextUtils;
+import android.util.Log;
+
+import com.android.vending.billing.IInAppBillingService;
+
+import org.json.JSONException;
+
+import java.util.ArrayList;
+import java.util.List;
+
+
+/**
+ * Provides convenience methods for in-app billing. You can create one instance of this
+ * class for your application and use it to process in-app billing operations.
+ * It provides synchronous (blocking) and asynchronous (non-blocking) methods for
+ * many common in-app billing operations, as well as automatic signature
+ * verification.
+ *
+ * After instantiating, you must perform setup in order to start using the object.
+ * To perform setup, call the {@link #startSetup} method and provide a listener;
+ * that listener will be notified when setup is complete, after which (and not before)
+ * you may call other methods.
+ *
+ * After setup is complete, you will typically want to request an inventory of owned
+ * items and subscriptions. See {@link #queryInventory}, {@link #queryInventoryAsync}
+ * and related methods.
+ *
+ * When you are done with this object, don't forget to call {@link #dispose}
+ * to ensure proper cleanup. This object holds a binding to the in-app billing
+ * service, which will leak unless you dispose of it correctly. If you created
+ * the object on an Activity's onCreate method, then the recommended
+ * place to dispose of it is the Activity's onDestroy method.
+ *
+ * A note about threading: When using this object from a background thread, you may
+ * call the blocking versions of methods; when using from a UI thread, call
+ * only the asynchronous versions and handle the results via callbacks.
+ * Also, notice that you can only call one asynchronous operation at a time;
+ * attempting to start a second asynchronous operation while the first one
+ * has not yet completed will result in an exception being thrown.
+ *
+ * @author Bruno Oliveira (Google)
+ *
+ */
+public class IabHelper {
+    // Is debug logging enabled?
+    boolean mDebugLog = false;
+    String mDebugTag = "IabHelper";
+
+    // Is setup done?
+    boolean mSetupDone = false;
+
+    // Has this object been disposed of? (If so, we should ignore callbacks, etc)
+    boolean mDisposed = false;
+
+    // Are subscriptions supported?
+    boolean mSubscriptionsSupported = false;
+
+    // Is an asynchronous operation in progress?
+    // (only one at a time can be in progress)
+    boolean mAsyncInProgress = false;
+
+    // (for logging/debugging)
+    // if mAsyncInProgress == true, what asynchronous operation is in progress?
+    String mAsyncOperation = "";
+
+    // Context we were passed during initialization
+    Context mContext;
+
+    // Connection to the service
+    IInAppBillingService mService;
+    ServiceConnection mServiceConn;
+
+    // The request code used to launch purchase flow
+    int mRequestCode;
+
+    // The item type of the current purchase flow
+    String mPurchasingItemType;
+
+    // Public key for verifying signature, in base64 encoding
+    String mSignatureBase64 = null;
+
+    // Billing response codes
+    public static final int BILLING_RESPONSE_RESULT_OK = 0;
+    public static final int BILLING_RESPONSE_RESULT_USER_CANCELED = 1;
+    public static final int BILLING_RESPONSE_RESULT_BILLING_UNAVAILABLE = 3;
+    public static final int BILLING_RESPONSE_RESULT_ITEM_UNAVAILABLE = 4;
+    public static final int BILLING_RESPONSE_RESULT_DEVELOPER_ERROR = 5;
+    public static final int BILLING_RESPONSE_RESULT_ERROR = 6;
+    public static final int BILLING_RESPONSE_RESULT_ITEM_ALREADY_OWNED = 7;
+    public static final int BILLING_RESPONSE_RESULT_ITEM_NOT_OWNED = 8;
+
+    // IAB Helper error codes
+    public static final int IABHELPER_ERROR_BASE = -1000;
+    public static final int IABHELPER_REMOTE_EXCEPTION = -1001;
+    public static final int IABHELPER_BAD_RESPONSE = -1002;
+    public static final int IABHELPER_VERIFICATION_FAILED = -1003;
+    public static final int IABHELPER_SEND_INTENT_FAILED = -1004;
+    public static final int IABHELPER_USER_CANCELLED = -1005;
+    public static final int IABHELPER_UNKNOWN_PURCHASE_RESPONSE = -1006;
+    public static final int IABHELPER_MISSING_TOKEN = -1007;
+    public static final int IABHELPER_UNKNOWN_ERROR = -1008;
+    public static final int IABHELPER_SUBSCRIPTIONS_NOT_AVAILABLE = -1009;
+    public static final int IABHELPER_INVALID_CONSUMPTION = -1010;
+
+    // Keys for the responses from InAppBillingService
+    public static final String RESPONSE_CODE = "RESPONSE_CODE";
+    public static final String RESPONSE_GET_SKU_DETAILS_LIST = "DETAILS_LIST";
+    public static final String RESPONSE_BUY_INTENT = "BUY_INTENT";
+    public static final String RESPONSE_INAPP_PURCHASE_DATA = "INAPP_PURCHASE_DATA";
+    public static final String RESPONSE_INAPP_SIGNATURE = "INAPP_DATA_SIGNATURE";
+    public static final String RESPONSE_INAPP_ITEM_LIST = "INAPP_PURCHASE_ITEM_LIST";
+    public static final String RESPONSE_INAPP_PURCHASE_DATA_LIST = "INAPP_PURCHASE_DATA_LIST";
+    public static final String RESPONSE_INAPP_SIGNATURE_LIST = "INAPP_DATA_SIGNATURE_LIST";
+    public static final String INAPP_CONTINUATION_TOKEN = "INAPP_CONTINUATION_TOKEN";
+
+    // Item types
+    public static final String ITEM_TYPE_INAPP = "inapp";
+    public static final String ITEM_TYPE_SUBS = "subs";
+
+    // some fields on the getSkuDetails response bundle
+    public static final String GET_SKU_DETAILS_ITEM_LIST = "ITEM_ID_LIST";
+    public static final String GET_SKU_DETAILS_ITEM_TYPE_LIST = "ITEM_TYPE_LIST";
+
+    /**
+     * Creates an instance. After creation, it will not yet be ready to use. You must perform
+     * setup by calling {@link #startSetup} and wait for setup to complete. This constructor does not
+     * block and is safe to call from a UI thread.
+     *
+     * @param ctx Your application or Activity context. Needed to bind to the in-app billing service.
+     * @param base64PublicKey Your application's public key, encoded in base64.
+     *     This is used for verification of purchase signatures. You can find your app's base64-encoded
+     *     public key in your application's page on Google Play Developer Console. Note that this
+     *     is NOT your "developer public key".
+     */
+    public IabHelper(Context ctx, String base64PublicKey) {
+        mContext = ctx.getApplicationContext();
+        mSignatureBase64 = base64PublicKey;
+        logDebug("IAB helper created.");
+    }
+
+    /**
+     * Enables or disable debug logging through LogCat.
+     */
+    public void enableDebugLogging(boolean enable, String tag) {
+        checkNotDisposed();
+        mDebugLog = enable;
+        mDebugTag = tag;
+    }
+
+    public void enableDebugLogging(boolean enable) {
+        checkNotDisposed();
+        mDebugLog = enable;
+    }
+
+    /**
+     * Callback for setup process. This listener's {@link #onIabSetupFinished} method is called
+     * when the setup process is complete.
+     */
+    public interface OnIabSetupFinishedListener {
+        /**
+         * Called to notify that setup is complete.
+         *
+         * @param result The result of the setup process.
+         */
+        public void onIabSetupFinished(IabResult result);
+    }
+
+    /**
+     * Starts the setup process. This will start up the setup process asynchronously.
+     * You will be notified through the listener when the setup process is complete.
+     * This method is safe to call from a UI thread.
+     *
+     * @param listener The listener to notify when the setup process is complete.
+     */
+    public void startSetup(final OnIabSetupFinishedListener listener) {
+        // If already set up, can't do it again.
+        checkNotDisposed();
+        if (mSetupDone) throw new IllegalStateException("IAB helper is already set up.");
+
+        // Connection to IAB service
+        logDebug("Starting in-app billing setup.");
+        mServiceConn = new ServiceConnection() {
+            @Override
+            public void onServiceDisconnected(ComponentName name) {
+                logDebug("Billing service disconnected.");
+                mService = null;
+            }
+
+            @Override
+            public void onServiceConnected(ComponentName name, IBinder service) {
+                if (mDisposed) return;
+                logDebug("Billing service connected.");
+                mService = IInAppBillingService.Stub.asInterface(service);
+                String packageName = mContext.getPackageName();
+                try {
+                    logDebug("Checking for in-app billing 3 support.");
+
+                    // check for in-app billing v3 support
+                    int response = mService.isBillingSupported(3, packageName, ITEM_TYPE_INAPP);
+                    if (response != BILLING_RESPONSE_RESULT_OK) {
+                        if (listener != null) listener.onIabSetupFinished(new IabResult(response,
+                                "Error checking for billing v3 support."));
+
+                        // if in-app purchases aren't supported, neither are subscriptions.
+                        mSubscriptionsSupported = false;
+                        return;
+                    }
+                    logDebug("In-app billing version 3 supported for " + packageName);
+
+                    // check for v3 subscriptions support
+                    response = mService.isBillingSupported(3, packageName, ITEM_TYPE_SUBS);
+                    if (response == BILLING_RESPONSE_RESULT_OK) {
+                        logDebug("Subscriptions AVAILABLE.");
+                        mSubscriptionsSupported = true;
+                    }
+                    else {
+                        logDebug("Subscriptions NOT AVAILABLE. Response: " + response);
+                    }
+
+                    mSetupDone = true;
+                }
+                catch (RemoteException e) {
+                    if (listener != null) {
+                        listener.onIabSetupFinished(new IabResult(IABHELPER_REMOTE_EXCEPTION,
+                                                    "RemoteException while setting up in-app billing."));
+                    }
+                    e.printStackTrace();
+                    return;
+                }
+
+                if (listener != null) {
+                    listener.onIabSetupFinished(new IabResult(BILLING_RESPONSE_RESULT_OK, "Setup successful."));
+                }
+            }
+        };
+
+        Intent serviceIntent = new Intent("com.android.vending.billing.InAppBillingService.BIND");
+        serviceIntent.setPackage("com.android.vending");
+        if (!mContext.getPackageManager().queryIntentServices(serviceIntent, 0).isEmpty()) {
+            // service available to handle that Intent
+            mContext.bindService(serviceIntent, mServiceConn, Context.BIND_AUTO_CREATE);
+        }
+        else {
+            // no service available to handle that Intent
+            if (listener != null) {
+                listener.onIabSetupFinished(
+                        new IabResult(BILLING_RESPONSE_RESULT_BILLING_UNAVAILABLE,
+                        "Billing service unavailable on device."));
+            }
+        }
+    }
+
+    /**
+     * Dispose of object, releasing resources. It's very important to call this
+     * method when you are done with this object. It will release any resources
+     * used by it such as service connections. Naturally, once the object is
+     * disposed of, it can't be used again.
+     */
+    public void dispose() {
+        logDebug("Disposing.");
+        mSetupDone = false;
+        if (mServiceConn != null) {
+            logDebug("Unbinding from service.");
+            if (mContext != null) mContext.unbindService(mServiceConn);
+        }
+        mDisposed = true;
+        mContext = null;
+        mServiceConn = null;
+        mService = null;
+        mPurchaseListener = null;
+    }
+
+    private void checkNotDisposed() {
+        if (mDisposed) throw new IllegalStateException("IabHelper was disposed of, so it cannot be used.");
+    }
+
+    /** Returns whether subscriptions are supported. */
+    public boolean subscriptionsSupported() {
+        checkNotDisposed();
+        return mSubscriptionsSupported;
+    }
+
+
+    /**
+     * Callback that notifies when a purchase is finished.
+     */
+    public interface OnIabPurchaseFinishedListener {
+        /**
+         * Called to notify that an in-app purchase finished. If the purchase was successful,
+         * then the sku parameter specifies which item was purchased. If the purchase failed,
+         * the sku and extraData parameters may or may not be null, depending on how far the purchase
+         * process went.
+         *
+         * @param result The result of the purchase.
+         * @param info The purchase information (null if purchase failed)
+         */
+        public void onIabPurchaseFinished(IabResult result, Purchase info);
+    }
+
+    // The listener registered on launchPurchaseFlow, which we have to call back when
+    // the purchase finishes
+    OnIabPurchaseFinishedListener mPurchaseListener;
+
+    public void launchPurchaseFlow(Activity act, String sku, int requestCode, OnIabPurchaseFinishedListener listener) {
+        launchPurchaseFlow(act, sku, requestCode, listener, "");
+    }
+
+    public void launchPurchaseFlow(Activity act, String sku, int requestCode,
+            OnIabPurchaseFinishedListener listener, String extraData) {
+        launchPurchaseFlow(act, sku, ITEM_TYPE_INAPP, requestCode, listener, extraData);
+    }
+
+    public void launchSubscriptionPurchaseFlow(Activity act, String sku, int requestCode,
+            OnIabPurchaseFinishedListener listener) {
+        launchSubscriptionPurchaseFlow(act, sku, requestCode, listener, "");
+    }
+
+    public void launchSubscriptionPurchaseFlow(Activity act, String sku, int requestCode,
+            OnIabPurchaseFinishedListener listener, String extraData) {
+        launchPurchaseFlow(act, sku, ITEM_TYPE_SUBS, requestCode, listener, extraData);
+    }
+
+    /**
+     * Initiate the UI flow for an in-app purchase. Call this method to initiate an in-app purchase,
+     * which will involve bringing up the Google Play screen. The calling activity will be paused while
+     * the user interacts with Google Play, and the result will be delivered via the activity's
+     * {@link android.app.Activity#onActivityResult} method, at which point you must call
+     * this object's {@link #handleActivityResult} method to continue the purchase flow. This method
+     * MUST be called from the UI thread of the Activity.
+     *
+     * @param act The calling activity.
+     * @param sku The sku of the item to purchase.
+     * @param itemType indicates if it's a product or a subscription (ITEM_TYPE_INAPP or ITEM_TYPE_SUBS)
+     * @param requestCode A request code (to differentiate from other responses --
+     *     as in {@link android.app.Activity#startActivityForResult}).
+     * @param listener The listener to notify when the purchase process finishes
+     * @param extraData Extra data (developer payload), which will be returned with the purchase data
+     *     when the purchase completes. This extra data will be permanently bound to that purchase
+     *     and will always be returned when the purchase is queried.
+     */
+    public void launchPurchaseFlow(Activity act, String sku, String itemType, int requestCode,
+                        OnIabPurchaseFinishedListener listener, String extraData) {
+        checkNotDisposed();
+        checkSetupDone("launchPurchaseFlow");
+        flagStartAsync("launchPurchaseFlow");
+        IabResult result;
+
+        if (itemType.equals(ITEM_TYPE_SUBS) && !mSubscriptionsSupported) {
+            IabResult r = new IabResult(IABHELPER_SUBSCRIPTIONS_NOT_AVAILABLE,
+                    "Subscriptions are not available.");
+            flagEndAsync();
+            if (listener != null) listener.onIabPurchaseFinished(r, null);
+            return;
+        }
+
+        try {
+            logDebug("Constructing buy intent for " + sku + ", item type: " + itemType);
+            Bundle buyIntentBundle = mService.getBuyIntent(3, mContext.getPackageName(), sku, itemType, extraData);
+            int response = getResponseCodeFromBundle(buyIntentBundle);
+            if (response != BILLING_RESPONSE_RESULT_OK) {
+                logError("Unable to buy item, Error response: " + getResponseDesc(response));
+                flagEndAsync();
+                result = new IabResult(response, "Unable to buy item");
+                if (listener != null) listener.onIabPurchaseFinished(result, null);
+                return;
+            }
+
+            PendingIntent pendingIntent = buyIntentBundle.getParcelable(RESPONSE_BUY_INTENT);
+            logDebug("Launching buy intent for " + sku + ". Request code: " + requestCode);
+            mRequestCode = requestCode;
+            mPurchaseListener = listener;
+            mPurchasingItemType = itemType;
+            act.startIntentSenderForResult(pendingIntent.getIntentSender(),
+                                           requestCode, new Intent(),
+                                           Integer.valueOf(0), Integer.valueOf(0),
+                                           Integer.valueOf(0));
+        }
+        catch (SendIntentException e) {
+            logError("SendIntentException while launching purchase flow for sku " + sku);
+            e.printStackTrace();
+            flagEndAsync();
+
+            result = new IabResult(IABHELPER_SEND_INTENT_FAILED, "Failed to send intent.");
+            if (listener != null) listener.onIabPurchaseFinished(result, null);
+        }
+        catch (RemoteException e) {
+            logError("RemoteException while launching purchase flow for sku " + sku);
+            e.printStackTrace();
+            flagEndAsync();
+
+            result = new IabResult(IABHELPER_REMOTE_EXCEPTION, "Remote exception while starting purchase flow");
+            if (listener != null) listener.onIabPurchaseFinished(result, null);
+        }
+    }
+
+    /**
+     * Handles an activity result that's part of the purchase flow in in-app billing. If you
+     * are calling {@link #launchPurchaseFlow}, then you must call this method from your
+     * Activity's {@link android.app.Activity@onActivityResult} method. This method
+     * MUST be called from the UI thread of the Activity.
+     *
+     * @param requestCode The requestCode as you received it.
+     * @param resultCode The resultCode as you received it.
+     * @param data The data (Intent) as you received it.
+     * @return Returns true if the result was related to a purchase flow and was handled;
+     *     false if the result was not related to a purchase, in which case you should
+     *     handle it normally.
+     */
+    public boolean handleActivityResult(int requestCode, int resultCode, Intent data) {
+        IabResult result;
+        if (requestCode != mRequestCode) return false;
+
+        checkNotDisposed();
+        checkSetupDone("handleActivityResult");
+
+        // end of async purchase operation that started on launchPurchaseFlow
+        flagEndAsync();
+
+        if (data == null) {
+            logError("Null data in IAB activity result.");
+            result = new IabResult(IABHELPER_BAD_RESPONSE, "Null data in IAB result");
+            if (mPurchaseListener != null) mPurchaseListener.onIabPurchaseFinished(result, null);
+            return true;
+        }
+
+        int responseCode = getResponseCodeFromIntent(data);
+        String purchaseData = data.getStringExtra(RESPONSE_INAPP_PURCHASE_DATA);
+        String dataSignature = data.getStringExtra(RESPONSE_INAPP_SIGNATURE);
+
+        if (resultCode == Activity.RESULT_OK && responseCode == BILLING_RESPONSE_RESULT_OK) {
+            logDebug("Successful resultcode from purchase activity.");
+            logDebug("Purchase data: " + purchaseData);
+            logDebug("Data signature: " + dataSignature);
+            logDebug("Extras: " + data.getExtras());
+            logDebug("Expected item type: " + mPurchasingItemType);
+
+            if (purchaseData == null || dataSignature == null) {
+                logError("BUG: either purchaseData or dataSignature is null.");
+                logDebug("Extras: " + data.getExtras().toString());
+                result = new IabResult(IABHELPER_UNKNOWN_ERROR, "IAB returned null purchaseData or dataSignature");
+                if (mPurchaseListener != null) mPurchaseListener.onIabPurchaseFinished(result, null);
+                return true;
+            }
+
+            Purchase purchase = null;
+            try {
+                purchase = new Purchase(mPurchasingItemType, purchaseData, dataSignature);
+                String sku = purchase.getSku();
+
+                // Verify signature
+                if (!Security.verifyPurchase(mSignatureBase64, purchaseData, dataSignature)) {
+                    logError("Purchase signature verification FAILED for sku " + sku);
+                    result = new IabResult(IABHELPER_VERIFICATION_FAILED, "Signature verification failed for sku " + sku);
+                    if (mPurchaseListener != null) mPurchaseListener.onIabPurchaseFinished(result, purchase);
+                    return true;
+                }
+                logDebug("Purchase signature successfully verified.");
+            }
+            catch (JSONException e) {
+                logError("Failed to parse purchase data.");
+                e.printStackTrace();
+                result = new IabResult(IABHELPER_BAD_RESPONSE, "Failed to parse purchase data.");
+                if (mPurchaseListener != null) mPurchaseListener.onIabPurchaseFinished(result, null);
+                return true;
+            }
+
+            if (mPurchaseListener != null) {
+                mPurchaseListener.onIabPurchaseFinished(new IabResult(BILLING_RESPONSE_RESULT_OK, "Success"), purchase);
+            }
+        }
+        else if (resultCode == Activity.RESULT_OK) {
+            // result code was OK, but in-app billing response was not OK.
+            logDebug("Result code was OK but in-app billing response was not OK: " + getResponseDesc(responseCode));
+            if (mPurchaseListener != null) {
+                result = new IabResult(responseCode, "Problem purchashing item.");
+                mPurchaseListener.onIabPurchaseFinished(result, null);
+            }
+        }
+        else if (resultCode == Activity.RESULT_CANCELED) {
+            logDebug("Purchase canceled - Response: " + getResponseDesc(responseCode));
+            result = new IabResult(IABHELPER_USER_CANCELLED, "User canceled.");
+            if (mPurchaseListener != null) mPurchaseListener.onIabPurchaseFinished(result, null);
+        }
+        else {
+            logError("Purchase failed. Result code: " + Integer.toString(resultCode)
+                    + ". Response: " + getResponseDesc(responseCode));
+            result = new IabResult(IABHELPER_UNKNOWN_PURCHASE_RESPONSE, "Unknown purchase response.");
+            if (mPurchaseListener != null) mPurchaseListener.onIabPurchaseFinished(result, null);
+        }
+        return true;
+    }
+
+    public Inventory queryInventory(boolean querySkuDetails, List<String> moreSkus) throws IabException {
+        return queryInventory(querySkuDetails, moreSkus, null);
+    }
+
+    /**
+     * Queries the inventory. This will query all owned items from the server, as well as
+     * information on additional skus, if specified. This method may block or take long to execute.
+     * Do not call from a UI thread. For that, use the non-blocking version {@link #refreshInventoryAsync}.
+     *
+     * @param querySkuDetails if true, SKU details (price, description, etc) will be queried as well
+     *     as purchase information.
+     * @param moreItemSkus additional PRODUCT skus to query information on, regardless of ownership.
+     *     Ignored if null or if querySkuDetails is false.
+     * @param moreSubsSkus additional SUBSCRIPTIONS skus to query information on, regardless of ownership.
+     *     Ignored if null or if querySkuDetails is false.
+     * @throws IabException if a problem occurs while refreshing the inventory.
+     */
+    public Inventory queryInventory(boolean querySkuDetails, List<String> moreItemSkus,
+                                        List<String> moreSubsSkus) throws IabException {
+        checkNotDisposed();
+        checkSetupDone("queryInventory");
+        try {
+            Inventory inv = new Inventory();
+            int r = queryPurchases(inv, ITEM_TYPE_INAPP);
+            if (r != BILLING_RESPONSE_RESULT_OK) {
+                throw new IabException(r, "Error refreshing inventory (querying owned items).");
+            }
+
+            if (querySkuDetails) {
+                r = querySkuDetails(ITEM_TYPE_INAPP, inv, moreItemSkus);
+                if (r != BILLING_RESPONSE_RESULT_OK) {
+                    throw new IabException(r, "Error refreshing inventory (querying prices of items).");
+                }
+            }
+
+            // if subscriptions are supported, then also query for subscriptions
+            if (mSubscriptionsSupported) {
+                r = queryPurchases(inv, ITEM_TYPE_SUBS);
+                if (r != BILLING_RESPONSE_RESULT_OK) {
+                    throw new IabException(r, "Error refreshing inventory (querying owned subscriptions).");
+                }
+
+                if (querySkuDetails) {
+                    r = querySkuDetails(ITEM_TYPE_SUBS, inv, moreItemSkus);
+                    if (r != BILLING_RESPONSE_RESULT_OK) {
+                        throw new IabException(r, "Error refreshing inventory (querying prices of subscriptions).");
+                    }
+                }
+            }
+
+            return inv;
+        }
+        catch (RemoteException e) {
+            throw new IabException(IABHELPER_REMOTE_EXCEPTION, "Remote exception while refreshing inventory.", e);
+        }
+        catch (JSONException e) {
+            throw new IabException(IABHELPER_BAD_RESPONSE, "Error parsing JSON response while refreshing inventory.", e);
+        }
+    }
+
+    /**
+     * Listener that notifies when an inventory query operation completes.
+     */
+    public interface QueryInventoryFinishedListener {
+        /**
+         * Called to notify that an inventory query operation completed.
+         *
+         * @param result The result of the operation.
+         * @param inv The inventory.
+         */
+        public void onQueryInventoryFinished(IabResult result, Inventory inv);
+    }
+
+
+    /**
+     * Asynchronous wrapper for inventory query. This will perform an inventory
+     * query as described in {@link #queryInventory}, but will do so asynchronously
+     * and call back the specified listener upon completion. This method is safe to
+     * call from a UI thread.
+     *
+     * @param querySkuDetails as in {@link #queryInventory}
+     * @param moreSkus as in {@link #queryInventory}
+     * @param listener The listener to notify when the refresh operation completes.
+     */
+    public void queryInventoryAsync(final boolean querySkuDetails,
+                               final List<String> moreSkus,
+                               final QueryInventoryFinishedListener listener) {
+        final Handler handler = new Handler();
+        checkNotDisposed();
+        checkSetupDone("queryInventory");
+        flagStartAsync("refresh inventory");
+        (new Thread(new Runnable() {
+            public void run() {
+                IabResult result = new IabResult(BILLING_RESPONSE_RESULT_OK, "Inventory refresh successful.");
+                Inventory inv = null;
+                try {
+                    inv = queryInventory(querySkuDetails, moreSkus);
+                }
+                catch (IabException ex) {
+                    result = ex.getResult();
+                }
+
+                flagEndAsync();
+
+                final IabResult result_f = result;
+                final Inventory inv_f = inv;
+                if (!mDisposed && listener != null) {
+                    handler.post(new Runnable() {
+                        public void run() {
+                            listener.onQueryInventoryFinished(result_f, inv_f);
+                        }
+                    });
+                }
+            }
+        })).start();
+    }
+
+    public void queryInventoryAsync(QueryInventoryFinishedListener listener) {
+        queryInventoryAsync(true, null, listener);
+    }
+
+    public void queryInventoryAsync(boolean querySkuDetails, QueryInventoryFinishedListener listener) {
+        queryInventoryAsync(querySkuDetails, null, listener);
+    }
+
+
+    /**
+     * Consumes a given in-app product. Consuming can only be done on an item
+     * that's owned, and as a result of consumption, the user will no longer own it.
+     * This method may block or take long to return. Do not call from the UI thread.
+     * For that, see {@link #consumeAsync}.
+     *
+     * @param itemInfo The PurchaseInfo that represents the item to consume.
+     * @throws IabException if there is a problem during consumption.
+     */
+    void consume(Purchase itemInfo) throws IabException {
+        checkNotDisposed();
+        checkSetupDone("consume");
+
+        if (!itemInfo.mItemType.equals(ITEM_TYPE_INAPP)) {
+            throw new IabException(IABHELPER_INVALID_CONSUMPTION,
+                    "Items of type '" + itemInfo.mItemType + "' can't be consumed.");
+        }
+
+        try {
+            String token = itemInfo.getToken();
+            String sku = itemInfo.getSku();
+            if (token == null || token.equals("")) {
+               logError("Can't consume "+ sku + ". No token.");
+               throw new IabException(IABHELPER_MISSING_TOKEN, "PurchaseInfo is missing token for sku: "
+                   + sku + " " + itemInfo);
+            }
+
+            logDebug("Consuming sku: " + sku + ", token: " + token);
+            int response = mService.consumePurchase(3, mContext.getPackageName(), token);
+            if (response == BILLING_RESPONSE_RESULT_OK) {
+               logDebug("Successfully consumed sku: " + sku);
+            }
+            else {
+               logDebug("Error consuming consuming sku " + sku + ". " + getResponseDesc(response));
+               throw new IabException(response, "Error consuming sku " + sku);
+            }
+        }
+        catch (RemoteException e) {
+            throw new IabException(IABHELPER_REMOTE_EXCEPTION, "Remote exception while consuming. PurchaseInfo: " + itemInfo, e);
+        }
+    }
+
+    /**
+     * Callback that notifies when a consumption operation finishes.
+     */
+    public interface OnConsumeFinishedListener {
+        /**
+         * Called to notify that a consumption has finished.
+         *
+         * @param purchase The purchase that was (or was to be) consumed.
+         * @param result The result of the consumption operation.
+         */
+        public void onConsumeFinished(Purchase purchase, IabResult result);
+    }
+
+    /**
+     * Callback that notifies when a multi-item consumption operation finishes.
+     */
+    public interface OnConsumeMultiFinishedListener {
+        /**
+         * Called to notify that a consumption of multiple items has finished.
+         *
+         * @param purchases The purchases that were (or were to be) consumed.
+         * @param results The results of each consumption operation, corresponding to each
+         *     sku.
+         */
+        public void onConsumeMultiFinished(List<Purchase> purchases, List<IabResult> results);
+    }
+
+    /**
+     * Asynchronous wrapper to item consumption. Works like {@link #consume}, but
+     * performs the consumption in the background and notifies completion through
+     * the provided listener. This method is safe to call from a UI thread.
+     *
+     * @param purchase The purchase to be consumed.
+     * @param listener The listener to notify when the consumption operation finishes.
+     */
+    public void consumeAsync(Purchase purchase, OnConsumeFinishedListener listener) {
+        checkNotDisposed();
+        checkSetupDone("consume");
+        List<Purchase> purchases = new ArrayList<Purchase>();
+        purchases.add(purchase);
+        consumeAsyncInternal(purchases, listener, null);
+    }
+
+    /**
+     * Same as {@link consumeAsync}, but for multiple items at once.
+     * @param purchases The list of PurchaseInfo objects representing the purchases to consume.
+     * @param listener The listener to notify when the consumption operation finishes.
+     */
+    public void consumeAsync(List<Purchase> purchases, OnConsumeMultiFinishedListener listener) {
+        checkNotDisposed();
+        checkSetupDone("consume");
+        consumeAsyncInternal(purchases, null, listener);
+    }
+
+    /**
+     * Returns a human-readable description for the given response code.
+     *
+     * @param code The response code
+     * @return A human-readable string explaining the result code.
+     *     It also includes the result code numerically.
+     */
+    public static String getResponseDesc(int code) {
+        String[] iab_msgs = ("0:OK/1:User Canceled/2:Unknown/" +
+                "3:Billing Unavailable/4:Item unavailable/" +
+                "5:Developer Error/6:Error/7:Item Already Owned/" +
+                "8:Item not owned").split("/");
+        String[] iabhelper_msgs = ("0:OK/-1001:Remote exception during initialization/" +
+                                   "-1002:Bad response received/" +
+                                   "-1003:Purchase signature verification failed/" +
+                                   "-1004:Send intent failed/" +
+                                   "-1005:User cancelled/" +
+                                   "-1006:Unknown purchase response/" +
+                                   "-1007:Missing token/" +
+                                   "-1008:Unknown error/" +
+                                   "-1009:Subscriptions not available/" +
+                                   "-1010:Invalid consumption attempt").split("/");
+
+        if (code <= IABHELPER_ERROR_BASE) {
+            int index = IABHELPER_ERROR_BASE - code;
+            if (index >= 0 && index < iabhelper_msgs.length) return iabhelper_msgs[index];
+            else return String.valueOf(code) + ":Unknown IAB Helper Error";
+        }
+        else if (code < 0 || code >= iab_msgs.length)
+            return String.valueOf(code) + ":Unknown";
+        else
+            return iab_msgs[code];
+    }
+
+
+    // Checks that setup was done; if not, throws an exception.
+    void checkSetupDone(String operation) {
+        if (!mSetupDone) {
+            logError("Illegal state for operation (" + operation + "): IAB helper is not set up.");
+            throw new IllegalStateException("IAB helper is not set up. Can't perform operation: " + operation);
+        }
+    }
+
+    // Workaround to bug where sometimes response codes come as Long instead of Integer
+    int getResponseCodeFromBundle(Bundle b) {
+        Object o = b.get(RESPONSE_CODE);
+        if (o == null) {
+            logDebug("Bundle with null response code, assuming OK (known issue)");
+            return BILLING_RESPONSE_RESULT_OK;
+        }
+        else if (o instanceof Integer) return ((Integer)o).intValue();
+        else if (o instanceof Long) return (int)((Long)o).longValue();
+        else {
+            logError("Unexpected type for bundle response code.");
+            logError(o.getClass().getName());
+            throw new RuntimeException("Unexpected type for bundle response code: " + o.getClass().getName());
+        }
+    }
+
+    // Workaround to bug where sometimes response codes come as Long instead of Integer
+    int getResponseCodeFromIntent(Intent i) {
+        Object o = i.getExtras().get(RESPONSE_CODE);
+        if (o == null) {
+            logError("Intent with no response code, assuming OK (known issue)");
+            return BILLING_RESPONSE_RESULT_OK;
+        }
+        else if (o instanceof Integer) return ((Integer)o).intValue();
+        else if (o instanceof Long) return (int)((Long)o).longValue();
+        else {
+            logError("Unexpected type for intent response code.");
+            logError(o.getClass().getName());
+            throw new RuntimeException("Unexpected type for intent response code: " + o.getClass().getName());
+        }
+    }
+
+    void flagStartAsync(String operation) {
+        if (mAsyncInProgress) throw new IllegalStateException("Can't start async operation (" +
+                operation + ") because another async operation(" + mAsyncOperation + ") is in progress.");
+        mAsyncOperation = operation;
+        mAsyncInProgress = true;
+        logDebug("Starting async operation: " + operation);
+    }
+
+    void flagEndAsync() {
+        logDebug("Ending async operation: " + mAsyncOperation);
+        mAsyncOperation = "";
+        mAsyncInProgress = false;
+    }
+
+
+    int queryPurchases(Inventory inv, String itemType) throws JSONException, RemoteException {
+        // Query purchases
+        logDebug("Querying owned items, item type: " + itemType);
+        logDebug("Package name: " + mContext.getPackageName());
+        boolean verificationFailed = false;
+        String continueToken = null;
+
+        do {
+            logDebug("Calling getPurchases with continuation token: " + continueToken);
+            Bundle ownedItems = mService.getPurchases(3, mContext.getPackageName(),
+                    itemType, continueToken);
+
+            int response = getResponseCodeFromBundle(ownedItems);
+            logDebug("Owned items response: " + String.valueOf(response));
+            if (response != BILLING_RESPONSE_RESULT_OK) {
+                logDebug("getPurchases() failed: " + getResponseDesc(response));
+                return response;
+            }
+            if (!ownedItems.containsKey(RESPONSE_INAPP_ITEM_LIST)
+                    || !ownedItems.containsKey(RESPONSE_INAPP_PURCHASE_DATA_LIST)
+                    || !ownedItems.containsKey(RESPONSE_INAPP_SIGNATURE_LIST)) {
+                logError("Bundle returned from getPurchases() doesn't contain required fields.");
+                return IABHELPER_BAD_RESPONSE;
+            }
+
+            ArrayList<String> ownedSkus = ownedItems.getStringArrayList(
+                        RESPONSE_INAPP_ITEM_LIST);
+            ArrayList<String> purchaseDataList = ownedItems.getStringArrayList(
+                        RESPONSE_INAPP_PURCHASE_DATA_LIST);
+            ArrayList<String> signatureList = ownedItems.getStringArrayList(
+                        RESPONSE_INAPP_SIGNATURE_LIST);
+
+            for (int i = 0; i < purchaseDataList.size(); ++i) {
+                String purchaseData = purchaseDataList.get(i);
+                String signature = signatureList.get(i);
+                String sku = ownedSkus.get(i);
+                if (Security.verifyPurchase(mSignatureBase64, purchaseData, signature)) {
+                    logDebug("Sku is owned: " + sku);
+                    Purchase purchase = new Purchase(itemType, purchaseData, signature);
+
+                    if (TextUtils.isEmpty(purchase.getToken())) {
+                        logWarn("BUG: empty/null token!");
+                        logDebug("Purchase data: " + purchaseData);
+                    }
+
+                    // Record ownership and token
+                    inv.addPurchase(purchase);
+                }
+                else {
+                    logWarn("Purchase signature verification **FAILED**. Not adding item.");
+                    logDebug("   Purchase data: " + purchaseData);
+                    logDebug("   Signature: " + signature);
+                    verificationFailed = true;
+                }
+            }
+
+            continueToken = ownedItems.getString(INAPP_CONTINUATION_TOKEN);
+            logDebug("Continuation token: " + continueToken);
+        } while (!TextUtils.isEmpty(continueToken));
+
+        return verificationFailed ? IABHELPER_VERIFICATION_FAILED : BILLING_RESPONSE_RESULT_OK;
+    }
+
+    int querySkuDetails(String itemType, Inventory inv, List<String> moreSkus)
+                                throws RemoteException, JSONException {
+        logDebug("Querying SKU details.");
+        ArrayList<String> skuList = new ArrayList<String>();
+        skuList.addAll(inv.getAllOwnedSkus(itemType));
+        if (moreSkus != null) {
+            for (String sku : moreSkus) {
+                if (!skuList.contains(sku)) {
+                    skuList.add(sku);
+                }
+            }
+        }
+
+        if (skuList.size() == 0) {
+            logDebug("queryPrices: nothing to do because there are no SKUs.");
+            return BILLING_RESPONSE_RESULT_OK;
+        }
+
+        Bundle querySkus = new Bundle();
+        querySkus.putStringArrayList(GET_SKU_DETAILS_ITEM_LIST, skuList);
+        Bundle skuDetails = mService.getSkuDetails(3, mContext.getPackageName(),
+                itemType, querySkus);
+
+        if (!skuDetails.containsKey(RESPONSE_GET_SKU_DETAILS_LIST)) {
+            int response = getResponseCodeFromBundle(skuDetails);
+            if (response != BILLING_RESPONSE_RESULT_OK) {
+                logDebug("getSkuDetails() failed: " + getResponseDesc(response));
+                return response;
+            }
+            else {
+                logError("getSkuDetails() returned a bundle with neither an error nor a detail list.");
+                return IABHELPER_BAD_RESPONSE;
+            }
+        }
+
+        ArrayList<String> responseList = skuDetails.getStringArrayList(
+                RESPONSE_GET_SKU_DETAILS_LIST);
+
+        for (String thisResponse : responseList) {
+            SkuDetails d = new SkuDetails(itemType, thisResponse);
+            logDebug("Got sku details: " + d);
+            inv.addSkuDetails(d);
+        }
+        return BILLING_RESPONSE_RESULT_OK;
+    }
+
+
+    void consumeAsyncInternal(final List<Purchase> purchases,
+                              final OnConsumeFinishedListener singleListener,
+                              final OnConsumeMultiFinishedListener multiListener) {
+        final Handler handler = new Handler();
+        flagStartAsync("consume");
+        (new Thread(new Runnable() {
+            public void run() {
+                final List<IabResult> results = new ArrayList<IabResult>();
+                for (Purchase purchase : purchases) {
+                    try {
+                        consume(purchase);
+                        results.add(new IabResult(BILLING_RESPONSE_RESULT_OK, "Successful consume of sku " + purchase.getSku()));
+                    }
+                    catch (IabException ex) {
+                        results.add(ex.getResult());
+                    }
+                }
+
+                flagEndAsync();
+                if (!mDisposed && singleListener != null) {
+                    handler.post(new Runnable() {
+                        public void run() {
+                            singleListener.onConsumeFinished(purchases.get(0), results.get(0));
+                        }
+                    });
+                }
+                if (!mDisposed && multiListener != null) {
+                    handler.post(new Runnable() {
+                        public void run() {
+                            multiListener.onConsumeMultiFinished(purchases, results);
+                        }
+                    });
+                }
+            }
+        })).start();
+    }
+
+    void logDebug(String msg) {
+        if (mDebugLog) Log.d(mDebugTag, msg);
+    }
+
+    void logError(String msg) {
+        Log.e(mDebugTag, "In-app billing error: " + msg);
+    }
+
+    void logWarn(String msg) {
+        Log.w(mDebugTag, "In-app billing warning: " + msg);
+    }
+}
diff --git a/Engine/Build/Android/Java/src/com/android/vending/billing/util/IabResult.java b/Engine/Build/Android/Java/src/com/android/vending/billing/util/IabResult.java
new file mode 100644
--- /dev/null
+++ b/Engine/Build/Android/Java/src/com/android/vending/billing/util/IabResult.java
@@ -0,0 +1,45 @@
+/* Copyright (c) 2012 Google Inc.
+ *
+ * Licensed under the Apache License, Version 2.0 (the "License");
+ * you may not use this file except in compliance with the License.
+ * You may obtain a copy of the License at
+ *
+ *     http://www.apache.org/licenses/LICENSE-2.0
+ *
+ * Unless required by applicable law or agreed to in writing, software
+ * distributed under the License is distributed on an "AS IS" BASIS,
+ * WITHOUT WARRANTIES OR CONDITIONS OF ANY KIND, either express or implied.
+ * See the License for the specific language governing permissions and
+ * limitations under the License.
+ */
+ 
+package com.android.vending.billing.util;
+
+/**
+ * Represents the result of an in-app billing operation.
+ * A result is composed of a response code (an integer) and possibly a
+ * message (String). You can get those by calling
+ * {@link #getResponse} and {@link #getMessage()}, respectively. You
+ * can also inquire whether a result is a success or a failure by
+ * calling {@link #isSuccess()} and {@link #isFailure()}.
+ */
+public class IabResult {
+    int mResponse;
+    String mMessage;
+
+    public IabResult(int response, String message) {
+        mResponse = response;
+        if (message == null || message.trim().length() == 0) {
+            mMessage = IabHelper.getResponseDesc(response);
+        }
+        else {
+            mMessage = message + " (response: " + IabHelper.getResponseDesc(response) + ")";
+        }
+    }
+    public int getResponse() { return mResponse; }
+    public String getMessage() { return mMessage; }
+    public boolean isSuccess() { return mResponse == IabHelper.BILLING_RESPONSE_RESULT_OK; }
+    public boolean isFailure() { return !isSuccess(); }
+    public String toString() { return "IabResult: " + getMessage(); }
+}
+
diff --git a/Engine/Build/Android/Java/src/com/android/vending/billing/util/Inventory.java b/Engine/Build/Android/Java/src/com/android/vending/billing/util/Inventory.java
new file mode 100644
--- /dev/null
+++ b/Engine/Build/Android/Java/src/com/android/vending/billing/util/Inventory.java
@@ -0,0 +1,91 @@
+/* Copyright (c) 2012 Google Inc.
+ *
+ * Licensed under the Apache License, Version 2.0 (the "License");
+ * you may not use this file except in compliance with the License.
+ * You may obtain a copy of the License at
+ *
+ *     http://www.apache.org/licenses/LICENSE-2.0
+ *
+ * Unless required by applicable law or agreed to in writing, software
+ * distributed under the License is distributed on an "AS IS" BASIS,
+ * WITHOUT WARRANTIES OR CONDITIONS OF ANY KIND, either express or implied.
+ * See the License for the specific language governing permissions and
+ * limitations under the License.
+ */
+ 
+package com.android.vending.billing.util;
+
+import java.util.ArrayList;
+import java.util.HashMap;
+import java.util.List;
+import java.util.Map;
+
+/**
+ * Represents a block of information about in-app items.
+ * An Inventory is returned by such methods as {@link IabHelper#queryInventory}.
+ */
+public class Inventory {
+    Map<String,SkuDetails> mSkuMap = new HashMap<String,SkuDetails>();
+    Map<String,Purchase> mPurchaseMap = new HashMap<String,Purchase>();
+
+    Inventory() { }
+
+    /** Returns the listing details for an in-app product. */
+    public SkuDetails getSkuDetails(String sku) {
+        return mSkuMap.get(sku);
+    }
+
+    /** Returns purchase information for a given product, or null if there is no purchase. */
+    public Purchase getPurchase(String sku) {
+        return mPurchaseMap.get(sku);
+    }
+
+    /** Returns whether or not there exists a purchase of the given product. */
+    public boolean hasPurchase(String sku) {
+        return mPurchaseMap.containsKey(sku);
+    }
+
+    /** Return whether or not details about the given product are available. */
+    public boolean hasDetails(String sku) {
+        return mSkuMap.containsKey(sku);
+    }
+
+    /**
+     * Erase a purchase (locally) from the inventory, given its product ID. This just
+     * modifies the Inventory object locally and has no effect on the server! This is
+     * useful when you have an existing Inventory object which you know to be up to date,
+     * and you have just consumed an item successfully, which means that erasing its
+     * purchase data from the Inventory you already have is quicker than querying for
+     * a new Inventory.
+     */
+    public void erasePurchase(String sku) {
+        if (mPurchaseMap.containsKey(sku)) mPurchaseMap.remove(sku);
+    }
+
+    /** Returns a list of all owned product IDs. */
+    List<String> getAllOwnedSkus() {
+        return new ArrayList<String>(mPurchaseMap.keySet());
+    }
+
+    /** Returns a list of all owned product IDs of a given type */
+    List<String> getAllOwnedSkus(String itemType) {
+        List<String> result = new ArrayList<String>();
+        for (Purchase p : mPurchaseMap.values()) {
+            if (p.getItemType().equals(itemType)) result.add(p.getSku());
+        }
+        return result;
+    }
+
+    /** Returns a list of all purchases. */
+    List<Purchase> getAllPurchases() {
+        return new ArrayList<Purchase>(mPurchaseMap.values());
+    }
+
+    void addSkuDetails(SkuDetails d) {
+        mSkuMap.put(d.getSku(), d);
+    }
+
+    void addPurchase(Purchase p) {
+        mPurchaseMap.put(p.getSku(), p);
+    }
+}
diff --git a/Engine/Build/Android/Java/src/com/android/vending/billing/util/Purchase.java b/Engine/Build/Android/Java/src/com/android/vending/billing/util/Purchase.java
new file mode 100644
--- /dev/null
+++ b/Engine/Build/Android/Java/src/com/android/vending/billing/util/Purchase.java
@@ -0,0 +1,63 @@
+/* Copyright (c) 2012 Google Inc.
+ *
+ * Licensed under the Apache License, Version 2.0 (the "License");
+ * you may not use this file except in compliance with the License.
+ * You may obtain a copy of the License at
+ *
+ *     http://www.apache.org/licenses/LICENSE-2.0
+ *
+ * Unless required by applicable law or agreed to in writing, software
+ * distributed under the License is distributed on an "AS IS" BASIS,
+ * WITHOUT WARRANTIES OR CONDITIONS OF ANY KIND, either express or implied.
+ * See the License for the specific language governing permissions and
+ * limitations under the License.
+ */
+ 
+package com.android.vending.billing.util;
+
+import org.json.JSONException;
+import org.json.JSONObject;
+
+/**
+ * Represents an in-app billing purchase.
+ */
+public class Purchase {
+    String mItemType;  // ITEM_TYPE_INAPP or ITEM_TYPE_SUBS
+    String mOrderId;
+    String mPackageName;
+    String mSku;
+    long mPurchaseTime;
+    int mPurchaseState;
+    String mDeveloperPayload;
+    String mToken;
+    String mOriginalJson;
+    String mSignature;
+
+    public Purchase(String itemType, String jsonPurchaseInfo, String signature) throws JSONException {
+        mItemType = itemType;
+        mOriginalJson = jsonPurchaseInfo;
+        JSONObject o = new JSONObject(mOriginalJson);
+        mOrderId = o.optString("orderId");
+        mPackageName = o.optString("packageName");
+        mSku = o.optString("productId");
+        mPurchaseTime = o.optLong("purchaseTime");
+        mPurchaseState = o.optInt("purchaseState");
+        mDeveloperPayload = o.optString("developerPayload");
+        mToken = o.optString("token", o.optString("purchaseToken"));
+        mSignature = signature;
+    }
+
+    public String getItemType() { return mItemType; }
+    public String getOrderId() { return mOrderId; }
+    public String getPackageName() { return mPackageName; }
+    public String getSku() { return mSku; }
+    public long getPurchaseTime() { return mPurchaseTime; }
+    public int getPurchaseState() { return mPurchaseState; }
+    public String getDeveloperPayload() { return mDeveloperPayload; }
+    public String getToken() { return mToken; }
+    public String getOriginalJson() { return mOriginalJson; }
+    public String getSignature() { return mSignature; }
+
+    @Override
+    public String toString() { return "PurchaseInfo(type:" + mItemType + "):" + mOriginalJson; }
+}
diff --git a/Engine/Build/Android/Java/src/com/android/vending/billing/util/Security.java b/Engine/Build/Android/Java/src/com/android/vending/billing/util/Security.java
new file mode 100644
--- /dev/null
+++ b/Engine/Build/Android/Java/src/com/android/vending/billing/util/Security.java
@@ -0,0 +1,123 @@
+/* Copyright (c) 2012 Google Inc.
+ *
+ * Licensed under the Apache License, Version 2.0 (the "License");
+ * you may not use this file except in compliance with the License.
+ * You may obtain a copy of the License at
+ *
+ *     http://www.apache.org/licenses/LICENSE-2.0
+ *
+ * Unless required by applicable law or agreed to in writing, software
+ * distributed under the License is distributed on an "AS IS" BASIS,
+ * WITHOUT WARRANTIES OR CONDITIONS OF ANY KIND, either express or implied.
+ * See the License for the specific language governing permissions and
+ * limitations under the License.
+ */
+ 
+package com.android.vending.billing.util;
+
+import android.text.TextUtils;
+import android.util.Log;
+
+import org.json.JSONException;
+import org.json.JSONObject;
+
+
+import java.security.InvalidKeyException;
+import java.security.KeyFactory;
+import java.security.NoSuchAlgorithmException;
+import java.security.PublicKey;
+import java.security.Signature;
+import java.security.SignatureException;
+import java.security.spec.InvalidKeySpecException;
+import java.security.spec.X509EncodedKeySpec;
+
+/**
+ * Security-related methods. For a secure implementation, all of this code
+ * should be implemented on a server that communicates with the
+ * application on the device. For the sake of simplicity and clarity of this
+ * example, this code is included here and is executed on the device. If you
+ * must verify the purchases on the phone, you should obfuscate this code to
+ * make it harder for an attacker to replace the code with stubs that treat all
+ * purchases as verified.
+ */
+public class Security {
+    private static final String TAG = "IABUtil/Security";
+
+    private static final String KEY_FACTORY_ALGORITHM = "RSA";
+    private static final String SIGNATURE_ALGORITHM = "SHA1withRSA";
+
+    /**
+     * Verifies that the data was signed with the given signature, and returns
+     * the verified purchase. The data is in JSON format and signed
+     * with a private key. The data also contains the {@link PurchaseState}
+     * and product ID of the purchase.
+     * @param base64PublicKey the base64-encoded public key to use for verifying.
+     * @param signedData the signed JSON string (signed, not encrypted)
+     * @param signature the signature for the data, signed with the private key
+     */
+    public static boolean verifyPurchase(String base64PublicKey, String signedData, String signature) {
+        if (TextUtils.isEmpty(signedData) || TextUtils.isEmpty(base64PublicKey) ||
+                TextUtils.isEmpty(signature)) {
+            Log.e(TAG, "Purchase verification failed: missing data.");
+            return false;
+        }
+
+        PublicKey key = Security.generatePublicKey(base64PublicKey);
+        return Security.verify(key, signedData, signature);
+    }
+
+    /**
+     * Generates a PublicKey instance from a string containing the
+     * Base64-encoded public key.
+     *
+     * @param encodedPublicKey Base64-encoded public key
+     * @throws IllegalArgumentException if encodedPublicKey is invalid
+     */
+    public static PublicKey generatePublicKey(String encodedPublicKey) {
+        try {
+            byte[] decodedKey = Base64.decode(encodedPublicKey);
+            KeyFactory keyFactory = KeyFactory.getInstance(KEY_FACTORY_ALGORITHM);
+            return keyFactory.generatePublic(new X509EncodedKeySpec(decodedKey));
+        } catch (NoSuchAlgorithmException e) {
+            throw new RuntimeException(e);
+        } catch (InvalidKeySpecException e) {
+            Log.e(TAG, "Invalid key specification.");
+            throw new IllegalArgumentException(e);
+        } catch (Base64DecoderException e) {
+            Log.e(TAG, "Base64 decoding failed.");
+            throw new IllegalArgumentException(e);
+        }
+    }
+
+    /**
+     * Verifies that the signature from the server matches the computed
+     * signature on the data.  Returns true if the data is correctly signed.
+     *
+     * @param publicKey public key associated with the developer account
+     * @param signedData signed data from server
+     * @param signature server signature
+     * @return true if the data and signature match
+     */
+    public static boolean verify(PublicKey publicKey, String signedData, String signature) {
+        Signature sig;
+        try {
+            sig = Signature.getInstance(SIGNATURE_ALGORITHM);
+            sig.initVerify(publicKey);
+            sig.update(signedData.getBytes());
+            if (!sig.verify(Base64.decode(signature))) {
+                Log.e(TAG, "Signature verification failed.");
+                return false;
+            }
+            return true;
+        } catch (NoSuchAlgorithmException e) {
+            Log.e(TAG, "NoSuchAlgorithmException.");
+        } catch (InvalidKeyException e) {
+            Log.e(TAG, "Invalid key specification.");
+        } catch (SignatureException e) {
+            Log.e(TAG, "Signature exception.");
+        } catch (Base64DecoderException e) {
+            Log.e(TAG, "Base64 decoding failed.");
+        }
+        return false;
+    }
+}
diff --git a/Engine/Build/Android/Java/src/com/android/vending/billing/util/SkuDetails.java b/Engine/Build/Android/Java/src/com/android/vending/billing/util/SkuDetails.java
new file mode 100644
--- /dev/null
+++ b/Engine/Build/Android/Java/src/com/android/vending/billing/util/SkuDetails.java
@@ -0,0 +1,58 @@
+/* Copyright (c) 2012 Google Inc.
+ *
+ * Licensed under the Apache License, Version 2.0 (the "License");
+ * you may not use this file except in compliance with the License.
+ * You may obtain a copy of the License at
+ *
+ *     http://www.apache.org/licenses/LICENSE-2.0
+ *
+ * Unless required by applicable law or agreed to in writing, software
+ * distributed under the License is distributed on an "AS IS" BASIS,
+ * WITHOUT WARRANTIES OR CONDITIONS OF ANY KIND, either express or implied.
+ * See the License for the specific language governing permissions and
+ * limitations under the License.
+ */
+ 
+package com.android.vending.billing.util;
+
+import org.json.JSONException;
+import org.json.JSONObject;
+
+/**
+ * Represents an in-app product's listing details.
+ */
+public class SkuDetails {
+    String mItemType;
+    String mSku;
+    String mType;
+    String mPrice;
+    String mTitle;
+    String mDescription;
+    String mJson;
+
+    public SkuDetails(String jsonSkuDetails) throws JSONException {
+        this(IabHelper.ITEM_TYPE_INAPP, jsonSkuDetails);
+    }
+
+    public SkuDetails(String itemType, String jsonSkuDetails) throws JSONException {
+        mItemType = itemType;
+        mJson = jsonSkuDetails;
+        JSONObject o = new JSONObject(mJson);
+        mSku = o.optString("productId");
+        mType = o.optString("type");
+        mPrice = o.optString("price");
+        mTitle = o.optString("title");
+        mDescription = o.optString("description");
+    }
+
+    public String getSku() { return mSku; }
+    public String getType() { return mType; }
+    public String getPrice() { return mPrice; }
+    public String getTitle() { return mTitle; }
+    public String getDescription() { return mDescription; }
+
+    @Override
+    public String toString() {
+        return "SkuDetails:" + mJson;
+    }
+}
diff --git a/Engine/Build/Android/Java/src/com/epicgames/ue4/GameActivity.java b/Engine/Build/Android/Java/src/com/epicgames/ue4/GameActivity.java
--- a/Engine/Build/Android/Java/src/com/epicgames/ue4/GameActivity.java
+++ b/Engine/Build/Android/Java/src/com/epicgames/ue4/GameActivity.java
@@ -58,6 +58,10 @@ import com.google.android.gms.plus.Plus;
 import java.net.URL;
 import java.net.HttpURLConnection;
 
+import com.epicgames.ue4.GooglePlayStoreHelper;
+import com.epicgames.ue4.GooglePlayLicensing;
+
+
 // TODO: use the resources from the UE4 lib project once we've got the packager up and running
 //import com.epicgames.ue4.R;
 import com.epicgames.ue4.JavaBuildSettings;
@@ -121,6 +125,8 @@ public class GameActivity extends NativeActivity
 	/** Unique ID to identify Google Play Services error dialog */
 	private static final int PLAY_SERVICES_DIALOG_ID = 1;
 
+	private StoreHelper IapStoreHelper;
+
 	@Override
 	public void onStart()
 	{
@@ -379,6 +385,9 @@ public class GameActivity extends NativeActivity
 		});
 		virtualKeyboardAlert = builder.create();
 
+		GooglePlayLicensing.GoogleLicensing = new GooglePlayLicensing();
+		GooglePlayLicensing.GoogleLicensing.Init(this, Log);
+
 		// Now okay for event handler to be set up on native side
 		nativeResumeMainInit();
 		
@@ -775,6 +784,77 @@ public class GameActivity extends NativeActivity
 		AudioManager audioManager = (AudioManager)this.getSystemService(Context.AUDIO_SERVICE);
 		return audioManager.isMusicActive();
 	}
+	
+	// In app purchase functionality
+	public void AndroidThunkJava_IapSetupService(String InProductKey)
+	{
+		Log.debug("[JAVA] - AndroidThunkJava_IapSetupService");
+		IapStoreHelper = new GooglePlayStoreHelper(InProductKey, this, Log);
+		if( IapStoreHelper == null )
+		{
+			Log.debug("[JAVA] - Store Helper is invalid");
+		}
+	}
+	
+
+	private String[] CachedQueryProductIDs;
+	private boolean[] CachedQueryConsumables;
+	public boolean AndroidThunkJava_IapQueryInAppPurchases(String[] ProductIDs, boolean[] bConsumable)
+	{
+		Log.debug("[JAVA] - AndroidThunkJava_IapQueryInAppPurchases");
+		CachedQueryProductIDs = ProductIDs;
+		CachedQueryConsumables = bConsumable;
+
+		boolean bTriggeredQuery = false;
+		if( IapStoreHelper != null )
+		{
+			bTriggeredQuery = true;
+
+			_activity.runOnUiThread(new Runnable()
+			{
+				@Override
+				public void run()
+				{
+					IapStoreHelper.QueryInAppPurchases(CachedQueryProductIDs, CachedQueryConsumables);
+				}
+			});
+		}
+		else
+		{
+			Log.debug("[JAVA] - Store Helper is invalid");
+		}
+		return bTriggeredQuery;
+	}
+	
+	public boolean AndroidThunkJava_IapBeginPurchase(String ProductId, boolean bConsumable)
+	{
+		Log.debug("[JAVA] - AndroidThunkJava_IapBeginPurchase");
+		boolean bTriggeredPurchase = false;
+		if( IapStoreHelper != null )
+		{
+			bTriggeredPurchase = IapStoreHelper.BeginPurchase(ProductId, bConsumable);
+		}
+		else
+		{
+			Log.debug("[JAVA] - Store Helper is invalid");
+		}
+		return bTriggeredPurchase;
+	}
+
+	public boolean AndroidThunkJava_IapIsAllowedToMakePurchases()
+	{
+		Log.debug("[JAVA] - AndroidThunkJava_IapIsAllowedToMakePurchases");
+		boolean bIsAllowedToMakePurchase = false;
+		if( IapStoreHelper != null )
+		{
+			bIsAllowedToMakePurchase = IapStoreHelper.IsAllowedToMakePurchases();
+		}
+		else
+		{
+			Log.debug("[JAVA] - Store Helper is invalid");
+		}
+		return bIsAllowedToMakePurchase;
+	}
 
 	public native boolean nativeIsShippingBuild();
 	public native void nativeSetGlobalActivity();
diff --git a/Engine/Build/Android/Java/src/com/epicgames/ue4/GooglePlayStoreHelper.java b/Engine/Build/Android/Java/src/com/epicgames/ue4/GooglePlayStoreHelper.java
new file mode 100644
--- /dev/null
+++ b/Engine/Build/Android/Java/src/com/epicgames/ue4/GooglePlayStoreHelper.java
@@ -0,0 +1,233 @@
+// Copyright 1998-2014 Epic Games, Inc. All Rights Reserved.
+//This file needs to be here so the "ant" build step doesnt fail when looking for a /src folder.
+
+package com.epicgames.ue4;
+
+import android.os.RemoteException;
+import android.content.IntentSender.SendIntentException;
+
+import com.android.vending.billing.util.IabHelper;
+import com.android.vending.billing.util.IabResult;
+import com.android.vending.billing.util.Inventory;
+import com.android.vending.billing.util.Purchase;
+import com.android.vending.billing.util.Base64;
+
+public class GooglePlayStoreHelper extends StoreHelper
+{
+	private IabHelper inAppPurchaseHelper;
+	
+	private GameActivity gameActivity;
+
+	private Logger Log;
+
+	private String ProductKey;
+	
+	// Iab helper functions for the store interface
+	IabHelper.OnConsumeFinishedListener ConsumeProductFinishedCallback =  new IabHelper.OnConsumeFinishedListener() 
+	{
+		public void onConsumeFinished(Purchase purchase, IabResult result) 
+		{
+			if (result.isSuccess()) 
+			{
+				Log.debug("[GooglePlayStoreHelper] - onConsumeFinished: " + (result.isSuccess() ? "Success" : "Failed" ));
+				String Receipt = Base64.encode(purchase.getOriginalJson().getBytes());
+				nativePurchaseComplete(true, purchase.getSku(), Receipt);
+			}
+			else
+			{
+				// Something went wrong, cannot consume product...send back success, but will be unable to purchase again
+				Log.debug("[GooglePlayStoreHelper] - ERROR: Could not successfully consume item. Will attempt again on next inventory update.");
+				nativePurchaseComplete(false, purchase.getSku(), "");
+			}
+		}
+	};
+
+	public GooglePlayStoreHelper(String InProductKey, GameActivity InGameActivity, final Logger InLog)
+	{
+		ProductKey = InProductKey;
+		gameActivity = InGameActivity;
+		Log = InLog;
+		SetupIapService();
+	}
+
+	private boolean bIsIapSetup = false;
+	public void SetupIapService()
+	{
+		Log.debug("[GooglePlayStoreHelper] - GooglePlayStoreHelper::SetupIapService");
+
+		inAppPurchaseHelper = new IabHelper(gameActivity, ProductKey);
+
+		inAppPurchaseHelper.startSetup(new IabHelper.OnIabSetupFinishedListener() 
+		{
+			public void onIabSetupFinished(IabResult result)
+			{
+				Log.debug("[GooglePlayStoreHelper] - GooglePlayStoreHelper::onIabSetupFinished");
+
+				bIsIapSetup = result.isSuccess();
+
+				if (bIsIapSetup)
+				{
+					Log.debug("[GooglePlayStoreHelper] - In-app purchasing helper has been setup");
+				}
+				else
+				{
+					Log.debug("[GooglePlayStoreHelper] - Problem setting up In-app purchasing helper: " + result);
+				}
+			}
+		});
+	}
+
+
+	private String[] cachedQueryProductIds;
+	private boolean[] cachedQueryConsumableFlags;
+	public boolean QueryInAppPurchases(final String[] ProductIds, final boolean[] bConsumableFlags)
+	{
+		Log.debug("[GooglePlayStoreHelper] - GooglePlayStoreHelper::QueryInAppPurchases");
+
+		// Cache the input parameters for cross threading
+		cachedQueryProductIds = ProductIds;
+		cachedQueryConsumableFlags = bConsumableFlags;
+		
+		// We can launch a purchase workflow if we have provided product Ids and the IabHelper is set-up.
+		boolean bIsAbleToSendQuery = bIsIapSetup && (cachedQueryProductIds.length > 0 );
+		if( bIsAbleToSendQuery )
+		{
+			IabHelper.QueryInventoryFinishedListener GotInventoryListener = new IabHelper.QueryInventoryFinishedListener()
+			{
+				public void onQueryInventoryFinished(IabResult result, Inventory inventory)
+				{
+					if (result.isSuccess())
+					{
+						Log.debug("[GooglePlayStoreHelper] - GooglePlayStoreHelper::onQueryInventoryFinished - Success.");
+						String[] titles = new String[cachedQueryProductIds.length];
+						String[] descriptions = new String[cachedQueryProductIds.length];
+						String[] prices = new String[cachedQueryProductIds.length];
+
+						int Idx = 0;
+						for( String productId : cachedQueryProductIds )
+						{
+							titles[Idx] = inventory.getSkuDetails(productId).getTitle();
+							descriptions[Idx] = inventory.getSkuDetails(productId).getDescription();
+							prices[Idx] = inventory.getSkuDetails(productId).getPrice();
+
+							++Idx;
+						}
+
+						Log.debug("[GooglePlayStoreHelper] - GooglePlayStoreHelper::onQueryInventoryFinished - Succeeded.");
+						nativeQueryComplete(true, cachedQueryProductIds, titles, descriptions, prices); 
+					}
+					else
+					{
+						Log.debug("[GooglePlayStoreHelper] - GooglePlayStoreHelper::onQueryInventoryFinished - Failed.");
+						nativeQueryComplete(false, null, null, null, null);
+					}
+				}
+			};
+
+			// Run the query pruchases command.
+			Log.debug("[GooglePlayStoreHelper] - queryInventoryAsync. - start");
+			try {
+				inAppPurchaseHelper.queryInventoryAsync(GotInventoryListener);
+			} catch(Exception e) {
+				Log.debug("[GooglePlayStoreHelper] - Exception caught - " + e.getMessage());
+			}
+			Log.debug("[GooglePlayStoreHelper] - queryInventoryAsync. - end");
+
+		}
+		else
+		{
+			if( bIsIapSetup == false )
+			{
+				Log.debug("[GooglePlayStoreHelper] - Failed to launch a query inventory flow, Store Helper is not set-up.");
+			}
+			else if(cachedQueryProductIds.length <= 0 )
+			{
+				Log.debug("[GooglePlayStoreHelper] - Failed to launch a query inventory flow, Cannot query without product Ids.");
+			}
+
+			nativeQueryComplete(false, null, null, null, null);
+		}
+
+
+		return bIsAbleToSendQuery;
+	}
+
+	String cachedPurchaseProductId;
+	boolean cachedPurchaseConsumableFlag;
+	public boolean BeginPurchase(final String ProductId, final boolean bConsumableFlag)
+	{
+		Log.debug("[GooglePlayStoreHelper] - GooglePlayStoreHelper::BeginPurchase");
+	
+		// Cache the input parameters for cross threading
+		cachedPurchaseProductId = ProductId;
+		cachedPurchaseConsumableFlag = bConsumableFlag;
+
+		// We can launch a purchase workflow if we have a provided product Id and the IabHelper is set-up.
+		boolean bIsValidProductId = (cachedPurchaseProductId.length() > 0);
+		if( bIsIapSetup && bIsValidProductId )
+		{
+			IabHelper.OnIabPurchaseFinishedListener PurchaseFinishedCallback = new IabHelper.OnIabPurchaseFinishedListener() 
+			{
+				public void onIabPurchaseFinished(IabResult result, Purchase purchase) 
+				{
+					if(result.isSuccess())
+					{
+						// Consumable items must be consumed on the server prior to provisioning the item to the game
+						if(cachedPurchaseConsumableFlag == true)
+						{
+							Log.debug("[GooglePlayStoreHelper] - Purchase Success. Requesting Item be consumed.");
+							inAppPurchaseHelper.consumeAsync(purchase, ConsumeProductFinishedCallback);
+						}
+						// Non consumable can immediately notify the game of the purchase
+						else
+						{
+							Log.debug("[GooglePlayStoreHelper] - Purchase Success.");
+							String Receipt = Base64.encode(purchase.getOriginalJson().getBytes());
+							nativePurchaseComplete(false, cachedPurchaseProductId, Receipt);
+						}
+					}
+					else
+					{
+						Log.debug("[GooglePlayStoreHelper] - Purchase Unsuccessful. " + result.getMessage());
+						nativePurchaseComplete(false, cachedPurchaseProductId, "" );
+					} 
+				}
+			};
+			
+			Log.debug("[GooglePlayStoreHelper] - launchPurchaseFlow. - start");
+			try {
+				inAppPurchaseHelper.launchPurchaseFlow(gameActivity, cachedPurchaseProductId, cachedPurchaseConsumableFlag ? 10001 : 10002, PurchaseFinishedCallback, "");
+			} catch(Exception e) {
+				Log.debug("[GooglePlayStoreHelper] - Exception caught - " + e.getMessage());
+			}
+			Log.debug("[GooglePlayStoreHelper] - launchPurchaseFlow. - end");
+		}
+		else
+		{
+			if( bIsIapSetup == false )
+			{
+				Log.debug("[GooglePlayStoreHelper] - Failed to launch a purchase flow, Store Helper is not set-up.");
+			}
+			else if( bIsValidProductId == false )
+			{
+				Log.debug("[GooglePlayStoreHelper] - Failed to launch a purchase flow, Cannot purchase without a valid product Id.");
+			}
+
+			nativePurchaseComplete(false, ProductId, "" );
+		}
+		
+		Log.debug("[GooglePlayStoreHelper] - GooglePlayStoreHelper::BeginPurchase - " + ((bIsIapSetup && bIsValidProductId) ? "true" : "false"));
+		return bIsIapSetup && bIsValidProductId;
+	}
+
+	public boolean IsAllowedToMakePurchases()
+	{
+		Log.debug("[GooglePlayStoreHelper] - GooglePlayStoreHelper::IsAllowedToMakePurchases");
+		return bIsIapSetup;
+	}
+
+	public native void nativeQueryComplete(boolean bSuccess, String[] productIDs, String[] titles, String[] descriptions, String[] prices );
+
+	public native void nativePurchaseComplete(boolean bSuccess, String ProductId, String ReceiptData);
+}
+
diff --git a/Engine/Build/Android/Java/src/com/epicgames/ue4/StoreHelper.java b/Engine/Build/Android/Java/src/com/epicgames/ue4/StoreHelper.java
new file mode 100644
--- /dev/null
+++ b/Engine/Build/Android/Java/src/com/epicgames/ue4/StoreHelper.java
@@ -0,0 +1,17 @@
+//This file needs to be here so the "ant" build step doesnt fail when looking for a /src folder.
+
+package com.epicgames.ue4;
+
+import android.content.Intent;
+
+public abstract class StoreHelper
+{
+	public StoreHelper()
+	{
+	}
+	
+	public abstract void SetupIapService();
+	public abstract boolean QueryInAppPurchases(String[] ProductIDs, boolean[] bConsumable);
+	public abstract boolean BeginPurchase(String ProductIDs, boolean bConsumable);
+	public abstract boolean IsAllowedToMakePurchases();
+}
\ No newline at end of file
diff --git a/Engine/Source/Runtime/Launch/Private/Android/AndroidJNI.cpp b/Engine/Source/Runtime/Launch/Private/Android/AndroidJNI.cpp
--- a/Engine/Source/Runtime/Launch/Private/Android/AndroidJNI.cpp
+++ b/Engine/Source/Runtime/Launch/Private/Android/AndroidJNI.cpp
@@ -121,6 +121,10 @@ jmethodID JDef_GameActivity::AndroidThunkJava_ForceQuit;
 jmethodID JDef_GameActivity::AndroidThunkJava_GetFontDirectory;
 jmethodID JDef_GameActivity::AndroidThunkJava_IsMusicActive;
 jmethodID JDef_GameActivity::AndroidThunkJava_InitHMDs;
+jmethodID JDef_GameActivity::AndroidThunkJava_IapSetupService;
+jmethodID JDef_GameActivity::AndroidThunkJava_IapQueryInAppPurchases;
+jmethodID JDef_GameActivity::AndroidThunkJava_IapBeginPurchase;
+jmethodID JDef_GameActivity::AndroidThunkJava_IapIsAllowedToMakePurchases;
 
 DEFINE_LOG_CATEGORY_STATIC(LogEngine, Log, All);
 
@@ -334,6 +338,107 @@ bool AndroidThunkCpp_IsMusicActive()
 	return active;
 }
 
+void AndroidThunkCpp_Iap_SetupIapService(const FString& InProductKey)
+{
+	if (JNIEnv* Env = GetJavaEnv())
+	{
+		if (GJavaGlobalThis)
+		{
+			jstring ProductKey = Env->NewStringUTF(TCHAR_TO_UTF8(*InProductKey));
+			Env->CallVoidMethod(GJavaGlobalThis, JDef_GameActivity::AndroidThunkJava_IapSetupService, ProductKey);
+			Env->DeleteLocalRef(ProductKey);
+		}
+		else
+		{
+			FPlatformMisc::LowLevelOutputDebugString(L"[JNI] - JNI is invalid");
+		}
+	}
+}
+
+bool AndroidThunkCpp_Iap_QueryInAppPurchases(const TArray<FString>& ProductIDs, const TArray<bool>& bConsumable)
+{
+	FPlatformMisc::LowLevelOutputDebugString(L"[JNI] - AndroidThunkCpp_Iap_QueryInAppPurchases");
+	bool bResult = false;
+
+	if (JNIEnv* Env = GetJavaEnv())
+	{
+		if (GJavaGlobalThis)
+		{
+			static jclass StringClass = Env->FindClass("java/lang/String");
+
+			// Populate some java types with the provided product information
+			jobjectArray ProductIDArray = (jobjectArray)Env->NewObjectArray(ProductIDs.Num(), StringClass, NULL);
+			jbooleanArray ConsumeArray = (jbooleanArray)Env->NewBooleanArray(ProductIDs.Num());
+
+			jboolean* ConsumeArrayValues = Env->GetBooleanArrayElements(ConsumeArray, 0);
+			for (uint32 Param = 0; Param < ProductIDs.Num(); Param++)
+			{
+				jstring StringValue = Env->NewStringUTF(TCHAR_TO_UTF8(*ProductIDs[Param]));
+				Env->SetObjectArrayElement(ProductIDArray, Param, StringValue);
+				Env->DeleteLocalRef(StringValue);
+
+				ConsumeArrayValues[Param] = bConsumable[Param];
+			}
+			Env->ReleaseBooleanArrayElements(ConsumeArray, ConsumeArrayValues, 0);
+
+			// Execute the java code for this operation
+			bResult = (bool)Env->CallBooleanMethod(GJavaGlobalThis, JDef_GameActivity::AndroidThunkJava_IapQueryInAppPurchases, ProductIDArray, ConsumeArray);
+
+			// clean up references
+			Env->DeleteLocalRef(ProductIDArray);
+			Env->DeleteLocalRef(ConsumeArray);
+		}
+		else
+		{
+			FPlatformMisc::LowLevelOutputDebugString(L"[JNI] - JNI is invalid");
+		}
+	}
+
+	return bResult;
+}
+
+bool AndroidThunkCpp_Iap_BeginPurchase(const FString& ProductID, const bool bConsumable)
+{
+	FPlatformMisc::LowLevelOutputDebugString(L"[JNI] - AndroidThunkCpp_Iap_BeginPurchase");
+	bool bResult = false;
+	if (JNIEnv* Env = GetJavaEnv())
+	{
+		if (GJavaGlobalThis)
+		{
+			jstring ProductIdJava = Env->NewStringUTF(TCHAR_TO_UTF8(*ProductID));
+			bResult = (bool)Env->CallBooleanMethod(GJavaGlobalThis, JDef_GameActivity::AndroidThunkJava_IapBeginPurchase, ProductIdJava, bConsumable);
+			Env->DeleteLocalRef(ProductIdJava);
+
+			FPlatformMisc::LowLevelOutputDebugStringf(TEXT("JDef_GameActivity::AndroidThunkJava_IapBeginPurchase - '%s'"), bResult ? TEXT("true") : TEXT("false"));
+		}
+		else
+		{
+			FPlatformMisc::LowLevelOutputDebugString(L"[JNI] - JNI is invalid");
+		}
+	}
+
+	return bResult;
+}
+
+bool AndroidThunkCpp_Iap_IsAllowedToMakePurchases()
+{
+	FPlatformMisc::LowLevelOutputDebugString(L"[JNI] - AndroidThunkCpp_Iap_IsAllowedToMakePurchases");
+	bool bResult = false;
+	if (JNIEnv* Env = GetJavaEnv())
+	{
+		if (GJavaGlobalThis)
+		{
+			bResult = (bool)Env->CallBooleanMethod(GJavaGlobalThis, JDef_GameActivity::AndroidThunkJava_IapIsAllowedToMakePurchases);
+		}
+		else
+		{
+			FPlatformMisc::LowLevelOutputDebugString(L"[JNI] - JNI is invalid");
+		}
+	}
+	return bResult;
+}
+
+
 //The JNI_OnLoad function is triggered by loading the game library from 
 //the Java source file.
 //	static
@@ -403,6 +508,19 @@ JNIEXPORT jint JNI_OnLoad(JavaVM* InJavaVM, void* InReserved)
 	JDef_GameActivity::AndroidThunkJava_InitHMDs = env->GetMethodID(JDef_GameActivity::ClassID, "AndroidThunkJava_InitHMDs", "()V");
 	CHECK_JNI_RESULT(JDef_GameActivity::AndroidThunkJava_InitHMDs);
 
+	// In app purchase functionality
+	JDef_GameActivity::AndroidThunkJava_IapSetupService = env->GetMethodID(JDef_GameActivity::ClassID, "AndroidThunkJava_IapSetupService", "(Ljava/lang/String;)V");
+	CHECK_JNI_RESULT(JDef_GameActivity::AndroidThunkJava_IapSetupService);
+
+	JDef_GameActivity::AndroidThunkJava_IapQueryInAppPurchases = env->GetMethodID(JDef_GameActivity::ClassID, "AndroidThunkJava_IapQueryInAppPurchases", "([Ljava/lang/String;[Z)Z");
+	CHECK_JNI_RESULT(JDef_GameActivity::AndroidThunkJava_IapQueryInAppPurchases);
+
+	JDef_GameActivity::AndroidThunkJava_IapBeginPurchase = env->GetMethodID(JDef_GameActivity::ClassID, "AndroidThunkJava_IapBeginPurchase", "(Ljava/lang/String;Z)Z");
+	CHECK_JNI_RESULT(JDef_GameActivity::AndroidThunkJava_IapBeginPurchase);
+
+	JDef_GameActivity::AndroidThunkJava_IapIsAllowedToMakePurchases = env->GetMethodID(JDef_GameActivity::ClassID, "AndroidThunkJava_IapIsAllowedToMakePurchases", "()Z");
+	CHECK_JNI_RESULT(JDef_GameActivity::AndroidThunkJava_IapIsAllowedToMakePurchases);
+
 	// hook signals
 #if UE_BUILD_DEBUG
 	if( GAlwaysReportCrash )
diff --git a/Engine/Source/Runtime/Launch/Private/Android/AndroidJNI.h b/Engine/Source/Runtime/Launch/Private/Android/AndroidJNI.h
--- a/Engine/Source/Runtime/Launch/Private/Android/AndroidJNI.h
+++ b/Engine/Source/Runtime/Launch/Private/Android/AndroidJNI.h
@@ -38,6 +38,12 @@ public:
 	static jmethodID AndroidThunkJava_GetFontDirectory;
 
 	static jmethodID AndroidThunkJava_IsMusicActive;
+
+	// In app purchase functionality
+	static jmethodID AndroidThunkJava_IapSetupService;
+	static jmethodID AndroidThunkJava_IapQueryInAppPurchases;
+	static jmethodID AndroidThunkJava_IapBeginPurchase;
+	static jmethodID AndroidThunkJava_IapIsAllowedToMakePurchases;
 };
 
 // Returns the java environment
diff --git a/Engine/Source/Runtime/Online/Android/OnlineSubsystemGooglePlay/Private/OnlineAsyncTaskGooglePlayQueryInAppPurchases.cpp b/Engine/Source/Runtime/Online/Android/OnlineSubsystemGooglePlay/Private/OnlineAsyncTaskGooglePlayQueryInAppPurchases.cpp
new file mode 100644
--- /dev/null
+++ b/Engine/Source/Runtime/Online/Android/OnlineSubsystemGooglePlay/Private/OnlineAsyncTaskGooglePlayQueryInAppPurchases.cpp
@@ -0,0 +1,48 @@
+// Copyright 1998-2014 Epic Games, Inc. All Rights Reserved.
+
+#include "OnlineSubsystemGooglePlayPrivatePCH.h"
+#include "OnlineAsyncTaskGooglePlayQueryInAppPurchases.h"
+
+FOnlineAsyncTaskGooglePlayQueryInAppPurchases::FOnlineAsyncTaskGooglePlayQueryInAppPurchases(
+	FOnlineSubsystemGooglePlay* InSubsystem,
+	const TArray<FString> InProductIds,
+	const TArray<bool> InIsConsumableFlags)
+	: FOnlineAsyncTaskBasic(InSubsystem)
+	, ProductIds(InProductIds)
+	, IsConsumableFlags(InIsConsumableFlags)
+
+{
+}
+
+
+void FOnlineAsyncTaskGooglePlayQueryInAppPurchases::ProcessQueryAvailablePurchasesResults(bool bInSuccessful)
+{
+	UE_LOG(LogOnline, Display, TEXT("FOnlineAsyncTaskGooglePlayQueryInAppPurchases::ProcessQueryAvailablePurchasesResults"));
+
+	bWasSuccessful = bInSuccessful;
+	bIsComplete = true;
+}
+
+
+void FOnlineAsyncTaskGooglePlayQueryInAppPurchases::Finalize()
+{
+	UE_LOG(LogOnline, Display, TEXT("FOnlineAsyncTaskGooglePlayQueryInAppPurchases::Finalize"));
+}
+
+
+void FOnlineAsyncTaskGooglePlayQueryInAppPurchases::TriggerDelegates()
+{
+	UE_LOG(LogOnline, Display, TEXT("FOnlineAsyncTaskGooglePlayQueryInAppPurchases::TriggerDelegates"));
+	Subsystem->GetStoreInterface()->TriggerOnQueryForAvailablePurchasesCompleteDelegates(bWasSuccessful);
+}
+
+
+void FOnlineAsyncTaskGooglePlayQueryInAppPurchases::Tick()
+{
+	UE_LOG(LogOnline, Display, TEXT("FOnlineAsyncTaskGooglePlayQueryInAppPurchases::Tick"));
+	bWasSuccessful = false;
+
+	extern bool AndroidThunkCpp_Iap_QueryInAppPurchases(const TArray<FString>&, const TArray<bool>&);
+	AndroidThunkCpp_Iap_QueryInAppPurchases(ProductIds, IsConsumableFlags);
+}
+
diff --git a/Engine/Source/Runtime/Online/Android/OnlineSubsystemGooglePlay/Private/OnlineAsyncTaskGooglePlayQueryInAppPurchases.h b/Engine/Source/Runtime/Online/Android/OnlineSubsystemGooglePlay/Private/OnlineAsyncTaskGooglePlayQueryInAppPurchases.h
new file mode 100644
--- /dev/null
+++ b/Engine/Source/Runtime/Online/Android/OnlineSubsystemGooglePlay/Private/OnlineAsyncTaskGooglePlayQueryInAppPurchases.h
@@ -0,0 +1,35 @@
+// Copyright 1998-2014 Epic Games, Inc. All Rights Reserved.
+
+#pragma once
+
+#include "OnlineAsyncTaskManager.h"
+
+class FOnlineSubsystemGooglePlay;
+
+class FOnlineAsyncTaskGooglePlayQueryInAppPurchases : public FOnlineAsyncTaskBasic<FOnlineSubsystemGooglePlay>
+{
+public:
+	FOnlineAsyncTaskGooglePlayQueryInAppPurchases(
+		FOnlineSubsystemGooglePlay* InSubsystem,
+		const TArray<FString> InProductIds,
+		const TArray<bool> InIsConsumableFlags);
+
+	// FOnlineAsyncItem
+	virtual FString ToString() const override { return TEXT("QueryInAppPurchases"); }
+	virtual void Finalize() override;
+	virtual void TriggerDelegates() override;
+
+	// FOnlineAsyncTask
+	virtual void Tick() override;
+
+	// FOnlineAsyncTaskGooglePlayQueryInAppPurchases
+	void ProcessQueryAvailablePurchasesResults(bool bInSuccessful);
+
+private:
+
+	// The product ids provided for this task
+	const TArray<FString> ProductIds;
+
+	// The consume flags provided for this task
+	const TArray<bool> IsConsumableFlags;
+};
diff --git a/Engine/Source/Runtime/Online/Android/OnlineSubsystemGooglePlay/Private/OnlineStoreInterfaceGooglePlay.cpp b/Engine/Source/Runtime/Online/Android/OnlineSubsystemGooglePlay/Private/OnlineStoreInterfaceGooglePlay.cpp
--- a/Engine/Source/Runtime/Online/Android/OnlineSubsystemGooglePlay/Private/OnlineStoreInterfaceGooglePlay.cpp
+++ b/Engine/Source/Runtime/Online/Android/OnlineSubsystemGooglePlay/Private/OnlineStoreInterfaceGooglePlay.cpp
@@ -1,48 +1,225 @@
 // Copyright 1998-2014 Epic Games, Inc. All Rights Reserved.
 
 #include "OnlineSubsystemGooglePlayPrivatePCH.h"
+#include "OnlineAsyncTaskGooglePlayQueryInAppPurchases.h"
+#include <jni.h>
 
 
 ////////////////////////////////////////////////////////////////////
 /// FOnlineStoreGooglePlay implementation
 
 
-FOnlineStoreGooglePlay::FOnlineStoreGooglePlay() 
+FOnlineStoreGooglePlay::FOnlineStoreGooglePlay(FOnlineSubsystemGooglePlay* InSubsystem)
+	: Subsystem( InSubsystem )
+	, CurrentQueryTask( nullptr )
 {
 	UE_LOG(LogOnline, Display, TEXT( "FOnlineStoreGooglePlay::FOnlineStoreGooglePlay" ));
+	
+	FString ProductKey;
+	GConfig->GetString(TEXT("/Script/AndroidRuntimeSettings.AndroidRuntimeSettings"), TEXT("GooglePlayProductKey"), ProductKey, GEngineIni);
 
+	extern void AndroidThunkCpp_Iap_SetupIapService(const FString&);
+	AndroidThunkCpp_Iap_SetupIapService(ProductKey);
 }
 
 
-FOnlineStoreGooglePlay::~FOnlineStoreGooglePlay() 
+FOnlineStoreGooglePlay::~FOnlineStoreGooglePlay()
 {
 	UE_LOG(LogOnline, Display, TEXT( "FOnlineStoreGooglePlay::~FOnlineStoreGooglePlay" ));
+}
+
 
+bool FOnlineStoreGooglePlay::IsAllowedToMakePurchases()
+{
+	UE_LOG(LogOnline, Display, TEXT("FOnlineStoreGooglePlay::IsAllowedToMakePurchases"));
+
+	extern bool AndroidThunkCpp_Iap_IsAllowedToMakePurchases();
+	return AndroidThunkCpp_Iap_IsAllowedToMakePurchases();
 }
 
 
-bool FOnlineStoreGooglePlay::QueryForAvailablePurchases(const TArray<FString>& ProductIDs, FOnlineProductInformationReadRef& InReadObject)
+bool FOnlineStoreGooglePlay::QueryForAvailablePurchases(const TArray<FString>& ProductIds, FOnlineProductInformationReadRef& InReadObject)
 {
 	UE_LOG(LogOnline, Display, TEXT( "FOnlineStoreGooglePlay::QueryForAvailablePurchases" ));
-	bool bSentAQueryRequest = false;
+	
+	ReadObject = InReadObject;
+	ReadObject->ReadState = EOnlineAsyncTaskState::InProgress;
+
+	TArray<bool> ConsumableFlags;
+	ConsumableFlags.AddZeroed(ProductIds.Num());
+
+	CurrentQueryTask = new FOnlineAsyncTaskGooglePlayQueryInAppPurchases(
+		Subsystem,
+		ProductIds,
+		ConsumableFlags);
+	Subsystem->QueueAsyncTask(CurrentQueryTask);
 
-	return bSentAQueryRequest;
+	return true;
+}
+
+
+extern "C" void Java_com_epicgames_ue4_GooglePlayStoreHelper_nativeQueryComplete(JNIEnv* jenv, jobject thiz, jboolean bSuccess, jobjectArray productIDs, jobjectArray titles, jobjectArray descriptions, jobjectArray prices)
+{
+	TArray<FInAppPurchaseProductInfo> ProvidedProductInformation;
+
+	if (jenv && bSuccess)
+	{
+		jsize NumProducts = jenv->GetArrayLength(productIDs);
+		jsize NumTitles = jenv->GetArrayLength(titles);
+		jsize NumDescriptions = jenv->GetArrayLength(descriptions);
+		jsize NumPrices = jenv->GetArrayLength(prices);
+
+		ensure((NumProducts == NumTitles) && (NumProducts == NumDescriptions) && (NumProducts == NumPrices));
+
+		for (jsize Idx = 0; Idx < NumProducts; Idx++)
+		{
+			// Build the product information strings.
+
+			FInAppPurchaseProductInfo NewProductInfo;
+
+			jstring NextId = (jstring)jenv->GetObjectArrayElement(productIDs, Idx);
+			const char* charsId = jenv->GetStringUTFChars(NextId, 0);
+			NewProductInfo.Identifier = FString(UTF8_TO_TCHAR(charsId));
+			jenv->ReleaseStringUTFChars(NextId, charsId);
+			jenv->DeleteLocalRef(NextId);
+
+			jstring NextTitle = (jstring)jenv->GetObjectArrayElement(titles, Idx);
+			const char* charsTitle = jenv->GetStringUTFChars(NextTitle, 0);
+			NewProductInfo.DisplayName = FString(UTF8_TO_TCHAR(charsTitle));
+			jenv->ReleaseStringUTFChars(NextTitle, charsTitle);
+			jenv->DeleteLocalRef(NextTitle);
+
+			jstring NextDesc = (jstring)jenv->GetObjectArrayElement(descriptions, Idx);
+			const char* charsDesc = jenv->GetStringUTFChars(NextDesc, 0);
+			NewProductInfo.DisplayDescription = FString(UTF8_TO_TCHAR(charsDesc));
+			jenv->ReleaseStringUTFChars(NextDesc, charsDesc);
+			jenv->DeleteLocalRef(NextDesc);
+
+			jstring NextPrice = (jstring)jenv->GetObjectArrayElement(prices, Idx);
+			const char* charsPrice = jenv->GetStringUTFChars(NextPrice, 0);
+			NewProductInfo.DisplayPrice = FString(UTF8_TO_TCHAR(charsPrice));
+			jenv->ReleaseStringUTFChars(NextPrice, charsPrice);
+			jenv->DeleteLocalRef(NextPrice);
+
+			ProvidedProductInformation.Add(NewProductInfo);
+
+			FPlatformMisc::LowLevelOutputDebugStringf(TEXT("\nProduct Identifier: %s, Name: %s, Description: %s, Price: %s\n"),
+				*NewProductInfo.Identifier,
+				*NewProductInfo.DisplayName,
+				*NewProductInfo.DisplayDescription,
+				*NewProductInfo.DisplayPrice);
+		}
+	}
+
+	if (IOnlineSubsystem* const OnlineSub = IOnlineSubsystem::Get())
+	{
+		FOnlineStoreGooglePlay* StoreInterface = (FOnlineStoreGooglePlay*)OnlineSub->GetStoreInterface().Get();
+		if (StoreInterface)
+		{
+			StoreInterface->ProcessQueryAvailablePurchasesResults(bSuccess, ProvidedProductInformation);
+		}
+	}
+
+	FPlatformMisc::LowLevelOutputDebugStringf(TEXT("In-App Purchase query was completed  %s\n"), bSuccess ? TEXT("successfully") : TEXT("unsuccessfully"));
+}
+
+
+void FOnlineStoreGooglePlay::ProcessQueryAvailablePurchasesResults(bool bInSuccessful, const TArray<FInAppPurchaseProductInfo>& AvailablePurchases)
+{
+	UE_LOG(LogOnline, Display, TEXT("FOnlineStoreGooglePlay::ProcessQueryAvailablePurchasesResults"));
+
+	if (ReadObject.IsValid())
+	{
+		ReadObject->ReadState = bInSuccessful ? EOnlineAsyncTaskState::Done : EOnlineAsyncTaskState::Failed;
+		ReadObject->ProvidedProductInformation.Insert(AvailablePurchases, 0);
+	}
+
+	CurrentQueryTask->ProcessQueryAvailablePurchasesResults(bInSuccessful);
 }
 
 
 bool FOnlineStoreGooglePlay::BeginPurchase(const FString& ProductId, FOnlineInAppPurchaseTransactionRef& InPurchaseStateObject)
 {
 	UE_LOG(LogOnline, Display, TEXT( "FOnlineStoreGooglePlay::BeginPurchase" ));
+	
 	bool bCreatedNewTransaction = false;
+	
+	if (IsAllowedToMakePurchases())
+	{
+		CachedPurchaseStateObject = InPurchaseStateObject;
+
+		extern bool AndroidThunkCpp_Iap_BeginPurchase(const FString&, const bool);
+		bCreatedNewTransaction = AndroidThunkCpp_Iap_BeginPurchase(ProductId, true); // TPMB - update begin purchase to use a consumable flag.
+		UE_LOG(LogOnline, Display, TEXT("Created Transaction? - %s"), 
+			bCreatedNewTransaction ? TEXT("Created a transaction.") : TEXT("Failed to create a transaction."));
+
+		if (!bCreatedNewTransaction)
+		{
+			UE_LOG(LogOnline, Display, TEXT("FOnlineStoreGooglePlay::BeginPurchase - Could not create a new transaction."));
+			CachedPurchaseStateObject->ReadState = EOnlineAsyncTaskState::Failed;
+			TriggerOnInAppPurchaseCompleteDelegates(EInAppPurchaseState::Invalid, nullptr);
+		}
+		else
+		{
+			CachedPurchaseStateObject->ReadState = EOnlineAsyncTaskState::InProgress;
+		}
+	}
+	else
+	{
+		UE_LOG(LogOnline, Display, TEXT("This device is not able to make purchases."));
+
+		InPurchaseStateObject->ReadState = EOnlineAsyncTaskState::Failed;
+		TriggerOnInAppPurchaseCompleteDelegates(EInAppPurchaseState::NotAllowed, nullptr);
+	}
+
 
 	return bCreatedNewTransaction;
 }
 
 
-bool FOnlineStoreGooglePlay::IsAllowedToMakePurchases()
+extern "C" void Java_com_epicgames_ue4_GooglePlayStoreHelper_nativePurchaseComplete(JNIEnv* jenv, jobject thiz, jboolean bSuccess, jstring productId, jstring receiptData)
 {
-	UE_LOG(LogOnline, Display, TEXT( "FOnlineStoreGooglePlay::IsAllowedToMakePurchases" ));
-	bool bCanMakePurchases = false;
+	if (IOnlineSubsystem* const OnlineSub = IOnlineSubsystem::Get())
+	{
+		// call store implementation to process query results.
+		if (FOnlineStoreGooglePlay* StoreInterface = (FOnlineStoreGooglePlay*)OnlineSub->GetStoreInterface().Get())
+		{
+			FString ProductId, ReceiptData;
+			if (bSuccess)
+			{
+				const char* charsId = jenv->GetStringUTFChars(productId, 0);
+				ProductId = FString(UTF8_TO_TCHAR(charsId));
+				jenv->ReleaseStringUTFChars(productId, charsId);
+
+				const char* charsReceipt = jenv->GetStringUTFChars(receiptData, 0);
+				ReceiptData = FString(UTF8_TO_TCHAR(charsReceipt));
+				jenv->ReleaseStringUTFChars(receiptData, charsReceipt);
+			}
+
+			StoreInterface->ProcessPurchaseResult(bSuccess, ProductId, ReceiptData);
+
+		}
+	}
+
+	FPlatformMisc::LowLevelOutputDebugStringf(TEXT("In-App Purchase was completed  %s\n"), bSuccess ? TEXT("successfully") : TEXT("unsuccessfully"));
+}
+
+
+void FOnlineStoreGooglePlay::ProcessPurchaseResult(bool bInSuccessful, const FString& ProductId, const FString& InReceiptData)
+{
+	UE_LOG(LogOnline, Display, TEXT("FOnlineStoreGooglePlay::ProcessPurchaseResult"));
+
+	FGooglePlayPurchaseReceipt Receipt;
+	if (ProductId.Len() > 0 && InReceiptData.Len() > 0)
+	{
+		Receipt.Identifier = ProductId;
+		Receipt.Data = InReceiptData;
+	}
+
+	if (CachedPurchaseStateObject.IsValid())
+	{
+		CachedPurchaseStateObject->ReadState = EOnlineAsyncTaskState::Done;
+	}
 
-	return bCanMakePurchases;
+	TriggerOnInAppPurchaseCompleteDelegates(bInSuccessful ? EInAppPurchaseState::Success : EInAppPurchaseState::Failed, &Receipt);
 }
\ No newline at end of file
diff --git a/Engine/Source/Runtime/Online/Android/OnlineSubsystemGooglePlay/Private/OnlineStoreInterfaceGooglePlay.h b/Engine/Source/Runtime/Online/Android/OnlineSubsystemGooglePlay/Private/OnlineStoreInterfaceGooglePlay.h
--- a/Engine/Source/Runtime/Online/Android/OnlineSubsystemGooglePlay/Private/OnlineStoreInterfaceGooglePlay.h
+++ b/Engine/Source/Runtime/Online/Android/OnlineSubsystemGooglePlay/Private/OnlineStoreInterfaceGooglePlay.h
@@ -5,6 +5,19 @@
 #include "OnlineStoreInterface.h"
 
 
+enum EGooglePlayBillingResponseCode
+{
+	Ok					= 0,
+	Cancelled			= 1,
+	BillingUnavailable	= 3,
+	ItemUnavailable		= 4,
+	DeveloperError		= 5,
+	Error				= 6,
+	ItemAlreadyOwned	= 7,
+	ItemNotOwned		= 8,
+};
+
+
 /**
  * The resulting state of an iap transaction
  */
@@ -19,13 +32,18 @@ namespace EInAppPurchaseResult
 	};
 }
 
+
 /**
  * Implementation of the Platform Purchase receipt. For this we provide an identifier and the encrypted data.
  */
 class FGooglePlayPurchaseReceipt : public IPlatformPurchaseReceipt
 {
 public:
+	// Product identifier
+	FString Identifier;
 
+	// The encrypted receipt data
+	FString Data;
 };
 
 
@@ -36,7 +54,7 @@ class FOnlineStoreGooglePlay : public IOnlineStore
 {
 public:
 	/** C-tor */
-	FOnlineStoreGooglePlay();
+	FOnlineStoreGooglePlay(FOnlineSubsystemGooglePlay* InSubsystem);
 	/** Destructor */
 	virtual ~FOnlineStoreGooglePlay();
 
@@ -46,13 +64,28 @@ public:
 	virtual bool IsAllowedToMakePurchases() override;
 	// End IOnlineStore 
 
+	void ProcessQueryAvailablePurchasesResults(bool bInSuccessful, const TArray<FInAppPurchaseProductInfo>& AvailablePurchases);
+	void ProcessPurchaseResult(bool bInSuccessful, const FString& InProductId, const FString& InReceiptData);
+
 private:
 
+	/** Pointer to owning subsystem */
+	FOnlineSubsystemGooglePlay* Subsystem;
+
+	/** The current query for iap async task */
+	class FOnlineAsyncTaskGooglePlayQueryInAppPurchases* CurrentQueryTask;
+
 	/** Delegate fired when a query for purchases has completed, whether successful or unsuccessful */
 	FOnQueryForAvailablePurchasesComplete OnQueryForAvailablePurchasesCompleteDelegate;
 
 	/** Delegate fired when a purchase transaction has completed, whether successful or unsuccessful */
 	FOnInAppPurchaseComplete OnPurchaseCompleteDelegate;
+
+	/** Cached in-app purchase query object, used to provide the user with product information attained from the server */
+	FOnlineProductInformationReadPtr ReadObject;
+
+	/** Cached in-app purchase transaction object, used to provide details to the user, of the product that has just been purchased. */
+	FOnlineInAppPurchaseTransactionPtr CachedPurchaseStateObject;
 };
 
 typedef TSharedPtr<FOnlineStoreGooglePlay, ESPMode::ThreadSafe> FOnlineStoreGooglePlayPtr;
diff --git a/Engine/Source/Runtime/Online/Android/OnlineSubsystemGooglePlay/Private/OnlineSubsystemGooglePlay.cpp b/Engine/Source/Runtime/Online/Android/OnlineSubsystemGooglePlay/Private/OnlineSubsystemGooglePlay.cpp
--- a/Engine/Source/Runtime/Online/Android/OnlineSubsystemGooglePlay/Private/OnlineSubsystemGooglePlay.cpp
+++ b/Engine/Source/Runtime/Online/Android/OnlineSubsystemGooglePlay/Private/OnlineSubsystemGooglePlay.cpp
@@ -20,6 +20,7 @@ FOnlineSubsystemGooglePlay::FOnlineSubsystemGooglePlay()
 	: IdentityInterface(nullptr)
 	, LeaderboardsInterface(nullptr)
 	, AchievementsInterface(nullptr)
+	, StoreInterface(nullptr)
 	, CurrentLoginTask(nullptr)
 {
 
@@ -114,8 +115,12 @@ bool FOnlineSubsystemGooglePlay::Init()
 	LeaderboardsInterface = MakeShareable(new FOnlineLeaderboardsGooglePlay(this));
 	AchievementsInterface = MakeShareable(new FOnlineAchievementsGooglePlay(this));
 	ExternalUIInterface = MakeShareable(new FOnlineExternalUIGooglePlay(this));
-	StoreInterface = MakeShareable(new FOnlineStoreGooglePlay());
 
+	if (IsInAppPurchasingEnabled())
+	{
+		StoreInterface = MakeShareable(new FOnlineStoreGooglePlay(this));
+	}
+	/*
 	extern struct android_app* GNativeAndroidApp;
 	check(GNativeAndroidApp != nullptr);
 	AndroidInitialization::android_main(GNativeAndroidApp);
@@ -139,7 +144,7 @@ bool FOnlineSubsystemGooglePlay::Init()
 			OnAuthActionFinished(Op, Status);
 		})
 		.Create(PlatformConfiguration).release() );
-
+*/
 	return true;
 }
 
@@ -166,6 +171,7 @@ bool FOnlineSubsystemGooglePlay::Shutdown()
 	}
 
 	// Destruct the interfaces
+	DESTRUCT_INTERFACE(StoreInterface);
 	DESTRUCT_INTERFACE(ExternalUIInterface);
 	DESTRUCT_INTERFACE(AchievementsInterface);
 	DESTRUCT_INTERFACE(LeaderboardsInterface);
@@ -198,6 +204,13 @@ bool FOnlineSubsystemGooglePlay::IsEnabled(void)
 	return bEnabled;
 }
 
+bool FOnlineSubsystemGooglePlay::IsInAppPurchasingEnabled()
+{
+	bool bEnabledIAP = false;
+	GConfig->GetBool(TEXT("OnlineSubsystemGooglePlay.Store"), TEXT("bSupportsInAppPurchasing"), bEnabledIAP, GEngineIni);
+	return bEnabledIAP;
+}
+
 void FOnlineSubsystemGooglePlay::QueueAsyncTask(FOnlineAsyncTask* AsyncTask)
 {
 	check(OnlineAsyncTaskThreadRunnable);
@@ -238,4 +251,4 @@ std::string FOnlineSubsystemGooglePlay::ConvertFStringToStdString(const FString&
 	FPlatformString::Convert(Converted.GetTypedData(), DestLen, *InString, SrcLen);
 
 	return std::string(Converted.GetTypedData());
-}
+}
\ No newline at end of file
diff --git a/Engine/Source/Runtime/Online/Android/OnlineSubsystemGooglePlay/Public/OnlineSubsystemGooglePlay.h b/Engine/Source/Runtime/Online/Android/OnlineSubsystemGooglePlay/Public/OnlineSubsystemGooglePlay.h
--- a/Engine/Source/Runtime/Online/Android/OnlineSubsystemGooglePlay/Public/OnlineSubsystemGooglePlay.h
+++ b/Engine/Source/Runtime/Online/Android/OnlineSubsystemGooglePlay/Public/OnlineSubsystemGooglePlay.h
@@ -91,6 +91,12 @@ PACKAGE_SCOPE:
 	/** Returns the Google Play-specific version of Achievements, useful to avoid unnecessary casting */
 	FOnlineAchievementsGooglePlayPtr GetAchievementsGooglePlay() const { return AchievementsInterface; }
 
+	/**
+	* Is IAP available for use
+	* @return true if IAP should be available, false otherwise
+	*/
+	bool IsInAppPurchasingEnabled();
+
 private:
 
 	/** Google callback when auth is complete */
