Missing file from CL 2345308

diff --git a/Engine/Build/Android/Java/src/com/epicgames/ue4/GooglePlayLicensing.java b/Engine/Build/Android/Java/src/com/epicgames/ue4/GooglePlayLicensing.java
new file mode 100644
--- /dev/null
+++ b/Engine/Build/Android/Java/src/com/epicgames/ue4/GooglePlayLicensing.java
@@ -0,0 +1,89 @@
+// Copyright 1998-2014 Epic Games, Inc. All Rights Reserved.
+// This file needs to be here so the "ant" build step doesnt fail when looking for a /src folder.
+
+package com.epicgames.ue4;
+
+import android.provider.Settings.Secure;
+
+
+import com.google.android.vending.licensing.AESObfuscator;
+import com.google.android.vending.licensing.LicenseChecker;
+import com.google.android.vending.licensing.LicenseCheckerCallback;
+import com.google.android.vending.licensing.Policy;
+import com.google.android.vending.licensing.ServerManagedPolicy;
+
+
+public class GooglePlayLicensing
+{
+	// TODO: Generate a new 20 byte list here
+	private static final byte[] SALT = new byte[] 
+	{
+		-46, 65, 30, -128, -103, -57, 74, -64, 51, 88, 
+		-95, -45, 77, -117, -36, -113, -11, 32, -64, 89
+	};
+
+	public static GooglePlayLicensing GoogleLicensing;
+	private GameActivity gameActivity;
+	private LicenseCheckerCallback mLicenseCheckerCallback;
+	private LicenseChecker mChecker;
+	private Logger Log;
+
+	public GooglePlayLicensing()
+	{
+	}
+
+	public void Init(GameActivity InActivity, final Logger InLog)
+	{
+		gameActivity = InActivity;
+		Log = InLog;
+	}
+
+	public void CheckLicense(String License)
+	{
+		Log.debug("Attempting to validate Google Play License.");
+		String deviceId = Secure.getString(gameActivity.getApplicationContext().getContentResolver(), Secure.ANDROID_ID);
+		mLicenseCheckerCallback = new MyLicenseCheckerCallback();
+		mChecker = new LicenseChecker(gameActivity.getApplicationContext(), new ServerManagedPolicy(gameActivity.getApplicationContext(), new AESObfuscator(SALT, gameActivity.getApplicationContext().getPackageName(), deviceId)), License);
+		mChecker.checkAccess(mLicenseCheckerCallback);
+	}
+
+	public void onDestroy()
+	{
+		if(mChecker != null)
+		{
+			mChecker.onDestroy();
+		}
+	}
+
+	private class MyLicenseCheckerCallback implements LicenseCheckerCallback 
+	{
+		public void allow(int policyReason) 
+		{
+			if (!gameActivity.isFinishing()) 
+			{
+				Log.debug("Game is Licensed version. Allowing access.");                
+			}
+		}
+
+		public void dontAllow(int policyReason) 
+		{
+			if (!gameActivity.isFinishing()) 
+			{
+				Log.debug("***************Game is not licensed!");
+//				gameActivity.getApplicationContext().ShowUnlicensedDialog();
+			}
+		}
+
+		public void applicationError(int errorCode) 
+		{
+			if (!gameActivity.isFinishing()) 
+			{
+				String result = Integer.toString(errorCode);
+				Log.debug("ERROR: " + result);
+			}
+		}
+	}
+
+	
+}
+
