Google licensing code. Was approved in a tps review. Needed for IAP.

diff --git a/Engine/Build/Android/Java/src/com/android/vending/licensing/AESObfuscator.java b/Engine/Build/Android/Java/src/com/android/vending/licensing/AESObfuscator.java
new file mode 100644
--- /dev/null
+++ b/Engine/Build/Android/Java/src/com/android/vending/licensing/AESObfuscator.java
@@ -0,0 +1,110 @@
+/*
+ * Copyright (C) 2010 The Android Open Source Project
+ *
+ * Licensed under the Apache License, Version 2.0 (the "License");
+ * you may not use this file except in compliance with the License.
+ * You may obtain a copy of the License at
+ *
+ *      http://www.apache.org/licenses/LICENSE-2.0
+ *
+ * Unless required by applicable law or agreed to in writing, software
+ * distributed under the License is distributed on an "AS IS" BASIS,
+ * WITHOUT WARRANTIES OR CONDITIONS OF ANY KIND, either express or implied.
+ * See the License for the specific language governing permissions and
+ * limitations under the License.
+ */
+
+package com.google.android.vending.licensing;
+
+import com.google.android.vending.licensing.util.Base64;
+import com.google.android.vending.licensing.util.Base64DecoderException;
+
+import java.io.UnsupportedEncodingException;
+import java.security.GeneralSecurityException;
+import java.security.spec.KeySpec;
+
+import javax.crypto.BadPaddingException;
+import javax.crypto.Cipher;
+import javax.crypto.IllegalBlockSizeException;
+import javax.crypto.SecretKey;
+import javax.crypto.SecretKeyFactory;
+import javax.crypto.spec.IvParameterSpec;
+import javax.crypto.spec.PBEKeySpec;
+import javax.crypto.spec.SecretKeySpec;
+
+/**
+ * An Obfuscator that uses AES to encrypt data.
+ */
+public class AESObfuscator implements Obfuscator {
+    private static final String UTF8 = "UTF-8";
+    private static final String KEYGEN_ALGORITHM = "PBEWITHSHAAND256BITAES-CBC-BC";
+    private static final String CIPHER_ALGORITHM = "AES/CBC/PKCS5Padding";
+    private static final byte[] IV =
+        { 16, 74, 71, -80, 32, 101, -47, 72, 117, -14, 0, -29, 70, 65, -12, 74 };
+    private static final String header = "com.android.vending.licensing.AESObfuscator-1|";
+
+    private Cipher mEncryptor;
+    private Cipher mDecryptor;
+
+    /**
+     * @param salt an array of random bytes to use for each (un)obfuscation
+     * @param applicationId application identifier, e.g. the package name
+     * @param deviceId device identifier. Use as many sources as possible to
+     *    create this unique identifier.
+     */
+    public AESObfuscator(byte[] salt, String applicationId, String deviceId) {
+        try {
+            SecretKeyFactory factory = SecretKeyFactory.getInstance(KEYGEN_ALGORITHM);
+            KeySpec keySpec =   
+                new PBEKeySpec((applicationId + deviceId).toCharArray(), salt, 1024, 256);
+            SecretKey tmp = factory.generateSecret(keySpec);
+            SecretKey secret = new SecretKeySpec(tmp.getEncoded(), "AES");
+            mEncryptor = Cipher.getInstance(CIPHER_ALGORITHM);
+            mEncryptor.init(Cipher.ENCRYPT_MODE, secret, new IvParameterSpec(IV));
+            mDecryptor = Cipher.getInstance(CIPHER_ALGORITHM);
+            mDecryptor.init(Cipher.DECRYPT_MODE, secret, new IvParameterSpec(IV));
+        } catch (GeneralSecurityException e) {
+            // This can't happen on a compatible Android device.
+            throw new RuntimeException("Invalid environment", e);
+        }
+    }
+
+    public String obfuscate(String original, String key) {
+        if (original == null) {
+            return null;
+        }
+        try {
+            // Header is appended as an integrity check
+            return Base64.encode(mEncryptor.doFinal((header + key + original).getBytes(UTF8)));
+        } catch (UnsupportedEncodingException e) {
+            throw new RuntimeException("Invalid environment", e);
+        } catch (GeneralSecurityException e) {
+            throw new RuntimeException("Invalid environment", e);
+        }
+    }
+
+    public String unobfuscate(String obfuscated, String key) throws ValidationException {
+        if (obfuscated == null) {
+            return null;
+        }
+        try {
+            String result = new String(mDecryptor.doFinal(Base64.decode(obfuscated)), UTF8);
+            // Check for presence of header. This serves as a final integrity check, for cases
+            // where the block size is correct during decryption.
+            int headerIndex = result.indexOf(header+key);
+            if (headerIndex != 0) {
+                throw new ValidationException("Header not found (invalid data or key)" + ":" +
+                        obfuscated);
+            }
+            return result.substring(header.length()+key.length(), result.length());
+        } catch (Base64DecoderException e) {
+            throw new ValidationException(e.getMessage() + ":" + obfuscated);
+        } catch (IllegalBlockSizeException e) {
+            throw new ValidationException(e.getMessage() + ":" + obfuscated);
+        } catch (BadPaddingException e) {
+            throw new ValidationException(e.getMessage() + ":" + obfuscated);
+        } catch (UnsupportedEncodingException e) {
+            throw new RuntimeException("Invalid environment", e);
+        }
+    }
+}
diff --git a/Engine/Build/Android/Java/src/com/android/vending/licensing/APKExpansionPolicy.java b/Engine/Build/Android/Java/src/com/android/vending/licensing/APKExpansionPolicy.java
new file mode 100644
--- /dev/null
+++ b/Engine/Build/Android/Java/src/com/android/vending/licensing/APKExpansionPolicy.java
@@ -0,0 +1,397 @@
+
+package com.google.android.vending.licensing;
+
+/*
+ * Copyright (C) 2012 The Android Open Source Project
+ *
+ * Licensed under the Apache License, Version 2.0 (the "License");
+ * you may not use this file except in compliance with the License.
+ * You may obtain a copy of the License at
+ *
+ *      http://www.apache.org/licenses/LICENSE-2.0
+ *
+ * Unless required by applicable law or agreed to in writing, software
+ * distributed under the License is distributed on an "AS IS" BASIS,
+ * WITHOUT WARRANTIES OR CONDITIONS OF ANY KIND, either express or implied.
+ * See the License for the specific language governing permissions and
+ * limitations under the License.
+ */
+
+import org.apache.http.NameValuePair;
+import org.apache.http.client.utils.URLEncodedUtils;
+
+import android.content.Context;
+import android.content.SharedPreferences;
+import android.util.Log;
+
+import java.net.URI;
+import java.net.URISyntaxException;
+import java.util.HashMap;
+import java.util.List;
+import java.util.Map;
+import java.util.Set;
+import java.util.Vector;
+
+/**
+ * Default policy. All policy decisions are based off of response data received
+ * from the licensing service. Specifically, the licensing server sends the
+ * following information: response validity period, error retry period, and
+ * error retry count.
+ * <p>
+ * These values will vary based on the the way the application is configured in
+ * the Android Market publishing console, such as whether the application is
+ * marked as free or is within its refund period, as well as how often an
+ * application is checking with the licensing service.
+ * <p>
+ * Developers who need more fine grained control over their application's
+ * licensing policy should implement a custom Policy.
+ */
+public class APKExpansionPolicy implements Policy {
+
+    private static final String TAG = "APKExpansionPolicy";
+    private static final String PREFS_FILE = "com.android.vending.licensing.APKExpansionPolicy";
+    private static final String PREF_LAST_RESPONSE = "lastResponse";
+    private static final String PREF_VALIDITY_TIMESTAMP = "validityTimestamp";
+    private static final String PREF_RETRY_UNTIL = "retryUntil";
+    private static final String PREF_MAX_RETRIES = "maxRetries";
+    private static final String PREF_RETRY_COUNT = "retryCount";
+    private static final String DEFAULT_VALIDITY_TIMESTAMP = "0";
+    private static final String DEFAULT_RETRY_UNTIL = "0";
+    private static final String DEFAULT_MAX_RETRIES = "0";
+    private static final String DEFAULT_RETRY_COUNT = "0";
+
+    private static final long MILLIS_PER_MINUTE = 60 * 1000;
+
+    private long mValidityTimestamp;
+    private long mRetryUntil;
+    private long mMaxRetries;
+    private long mRetryCount;
+    private long mLastResponseTime = 0;
+    private int mLastResponse;
+    private PreferenceObfuscator mPreferences;
+    private Vector<String> mExpansionURLs = new Vector<String>();
+    private Vector<String> mExpansionFileNames = new Vector<String>();
+    private Vector<Long> mExpansionFileSizes = new Vector<Long>();
+
+    /**
+     * The design of the protocol supports n files. Currently the market can
+     * only deliver two files. To accommodate this, we have these two constants,
+     * but the order is the only relevant thing here.
+     */
+    public static final int MAIN_FILE_URL_INDEX = 0;
+    public static final int PATCH_FILE_URL_INDEX = 1;
+
+    /**
+     * @param context The context for the current application
+     * @param obfuscator An obfuscator to be used with preferences.
+     */
+    public APKExpansionPolicy(Context context, Obfuscator obfuscator) {
+        // Import old values
+        SharedPreferences sp = context.getSharedPreferences(PREFS_FILE, Context.MODE_PRIVATE);
+        mPreferences = new PreferenceObfuscator(sp, obfuscator);
+        mLastResponse = Integer.parseInt(
+                mPreferences.getString(PREF_LAST_RESPONSE, Integer.toString(Policy.RETRY)));
+        mValidityTimestamp = Long.parseLong(mPreferences.getString(PREF_VALIDITY_TIMESTAMP,
+                DEFAULT_VALIDITY_TIMESTAMP));
+        mRetryUntil = Long.parseLong(mPreferences.getString(PREF_RETRY_UNTIL, DEFAULT_RETRY_UNTIL));
+        mMaxRetries = Long.parseLong(mPreferences.getString(PREF_MAX_RETRIES, DEFAULT_MAX_RETRIES));
+        mRetryCount = Long.parseLong(mPreferences.getString(PREF_RETRY_COUNT, DEFAULT_RETRY_COUNT));
+    }
+
+    /**
+     * We call this to guarantee that we fetch a fresh policy from the server.
+     * This is to be used if the URL is invalid.
+     */
+    public void resetPolicy() {
+        mPreferences.putString(PREF_LAST_RESPONSE, Integer.toString(Policy.RETRY));
+        setRetryUntil(DEFAULT_RETRY_UNTIL);
+        setMaxRetries(DEFAULT_MAX_RETRIES);
+        setRetryCount(Long.parseLong(DEFAULT_RETRY_COUNT));
+        setValidityTimestamp(DEFAULT_VALIDITY_TIMESTAMP);
+        mPreferences.commit();
+    }
+
+    /**
+     * Process a new response from the license server.
+     * <p>
+     * This data will be used for computing future policy decisions. The
+     * following parameters are processed:
+     * <ul>
+     * <li>VT: the timestamp that the client should consider the response valid
+     * until
+     * <li>GT: the timestamp that the client should ignore retry errors until
+     * <li>GR: the number of retry errors that the client should ignore
+     * </ul>
+     * 
+     * @param response the result from validating the server response
+     * @param rawData the raw server response data
+     */
+    public void processServerResponse(int response,
+            com.google.android.vending.licensing.ResponseData rawData) {
+
+        // Update retry counter
+        if (response != Policy.RETRY) {
+            setRetryCount(0);
+        } else {
+            setRetryCount(mRetryCount + 1);
+        }
+
+        if (response == Policy.LICENSED) {
+            // Update server policy data
+            Map<String, String> extras = decodeExtras(rawData.extra);
+            mLastResponse = response;
+            setValidityTimestamp(Long.toString(System.currentTimeMillis() + MILLIS_PER_MINUTE));
+            Set<String> keys = extras.keySet();
+            for (String key : keys) {
+                if (key.equals("VT")) {
+                    setValidityTimestamp(extras.get(key));
+                } else if (key.equals("GT")) {
+                    setRetryUntil(extras.get(key));
+                } else if (key.equals("GR")) {
+                    setMaxRetries(extras.get(key));
+                } else if (key.startsWith("FILE_URL")) {
+                    int index = Integer.parseInt(key.substring("FILE_URL".length())) - 1;
+                    setExpansionURL(index, extras.get(key));
+                } else if (key.startsWith("FILE_NAME")) {
+                    int index = Integer.parseInt(key.substring("FILE_NAME".length())) - 1;
+                    setExpansionFileName(index, extras.get(key));
+                } else if (key.startsWith("FILE_SIZE")) {
+                    int index = Integer.parseInt(key.substring("FILE_SIZE".length())) - 1;
+                    setExpansionFileSize(index, Long.parseLong(extras.get(key)));
+                }
+            }
+        } else if (response == Policy.NOT_LICENSED) {
+            // Clear out stale policy data
+            setValidityTimestamp(DEFAULT_VALIDITY_TIMESTAMP);
+            setRetryUntil(DEFAULT_RETRY_UNTIL);
+            setMaxRetries(DEFAULT_MAX_RETRIES);
+        }
+
+        setLastResponse(response);
+        mPreferences.commit();
+    }
+
+    /**
+     * Set the last license response received from the server and add to
+     * preferences. You must manually call PreferenceObfuscator.commit() to
+     * commit these changes to disk.
+     * 
+     * @param l the response
+     */
+    private void setLastResponse(int l) {
+        mLastResponseTime = System.currentTimeMillis();
+        mLastResponse = l;
+        mPreferences.putString(PREF_LAST_RESPONSE, Integer.toString(l));
+    }
+
+    /**
+     * Set the current retry count and add to preferences. You must manually
+     * call PreferenceObfuscator.commit() to commit these changes to disk.
+     * 
+     * @param c the new retry count
+     */
+    private void setRetryCount(long c) {
+        mRetryCount = c;
+        mPreferences.putString(PREF_RETRY_COUNT, Long.toString(c));
+    }
+
+    public long getRetryCount() {
+        return mRetryCount;
+    }
+
+    /**
+     * Set the last validity timestamp (VT) received from the server and add to
+     * preferences. You must manually call PreferenceObfuscator.commit() to
+     * commit these changes to disk.
+     * 
+     * @param validityTimestamp the VT string received
+     */
+    private void setValidityTimestamp(String validityTimestamp) {
+        Long lValidityTimestamp;
+        try {
+            lValidityTimestamp = Long.parseLong(validityTimestamp);
+        } catch (NumberFormatException e) {
+            // No response or not parseable, expire in one minute.
+            Log.w(TAG, "License validity timestamp (VT) missing, caching for a minute");
+            lValidityTimestamp = System.currentTimeMillis() + MILLIS_PER_MINUTE;
+            validityTimestamp = Long.toString(lValidityTimestamp);
+        }
+
+        mValidityTimestamp = lValidityTimestamp;
+        mPreferences.putString(PREF_VALIDITY_TIMESTAMP, validityTimestamp);
+    }
+
+    public long getValidityTimestamp() {
+        return mValidityTimestamp;
+    }
+
+    /**
+     * Set the retry until timestamp (GT) received from the server and add to
+     * preferences. You must manually call PreferenceObfuscator.commit() to
+     * commit these changes to disk.
+     * 
+     * @param retryUntil the GT string received
+     */
+    private void setRetryUntil(String retryUntil) {
+        Long lRetryUntil;
+        try {
+            lRetryUntil = Long.parseLong(retryUntil);
+        } catch (NumberFormatException e) {
+            // No response or not parseable, expire immediately
+            Log.w(TAG, "License retry timestamp (GT) missing, grace period disabled");
+            retryUntil = "0";
+            lRetryUntil = 0l;
+        }
+
+        mRetryUntil = lRetryUntil;
+        mPreferences.putString(PREF_RETRY_UNTIL, retryUntil);
+    }
+
+    public long getRetryUntil() {
+        return mRetryUntil;
+    }
+
+    /**
+     * Set the max retries value (GR) as received from the server and add to
+     * preferences. You must manually call PreferenceObfuscator.commit() to
+     * commit these changes to disk.
+     * 
+     * @param maxRetries the GR string received
+     */
+    private void setMaxRetries(String maxRetries) {
+        Long lMaxRetries;
+        try {
+            lMaxRetries = Long.parseLong(maxRetries);
+        } catch (NumberFormatException e) {
+            // No response or not parseable, expire immediately
+            Log.w(TAG, "Licence retry count (GR) missing, grace period disabled");
+            maxRetries = "0";
+            lMaxRetries = 0l;
+        }
+
+        mMaxRetries = lMaxRetries;
+        mPreferences.putString(PREF_MAX_RETRIES, maxRetries);
+    }
+
+    public long getMaxRetries() {
+        return mMaxRetries;
+    }
+
+    /**
+     * Gets the count of expansion URLs. Since expansionURLs are not committed
+     * to preferences, this will return zero if there has been no LVL fetch
+     * in the current session.
+     * 
+     * @return the number of expansion URLs. (0,1,2)
+     */
+    public int getExpansionURLCount() {
+        return mExpansionURLs.size();
+    }
+
+    /**
+     * Gets the expansion URL. Since these URLs are not committed to
+     * preferences, this will always return null if there has not been an LVL
+     * fetch in the current session.
+     * 
+     * @param index the index of the URL to fetch. This value will be either
+     *            MAIN_FILE_URL_INDEX or PATCH_FILE_URL_INDEX
+     * @param URL the URL to set
+     */
+    public String getExpansionURL(int index) {
+        if (index < mExpansionURLs.size()) {
+            return mExpansionURLs.elementAt(index);
+        }
+        return null;
+    }
+
+    /**
+     * Sets the expansion URL. Expansion URL's are not committed to preferences,
+     * but are instead intended to be stored when the license response is
+     * processed by the front-end.
+     * 
+     * @param index the index of the expansion URL. This value will be either
+     *            MAIN_FILE_URL_INDEX or PATCH_FILE_URL_INDEX
+     * @param URL the URL to set
+     */
+    public void setExpansionURL(int index, String URL) {
+        if (index >= mExpansionURLs.size()) {
+            mExpansionURLs.setSize(index + 1);
+        }
+        mExpansionURLs.set(index, URL);
+    }
+
+    public String getExpansionFileName(int index) {
+        if (index < mExpansionFileNames.size()) {
+            return mExpansionFileNames.elementAt(index);
+        }
+        return null;
+    }
+
+    public void setExpansionFileName(int index, String name) {
+        if (index >= mExpansionFileNames.size()) {
+            mExpansionFileNames.setSize(index + 1);
+        }
+        mExpansionFileNames.set(index, name);
+    }
+
+    public long getExpansionFileSize(int index) {
+        if (index < mExpansionFileSizes.size()) {
+            return mExpansionFileSizes.elementAt(index);
+        }
+        return -1;
+    }
+
+    public void setExpansionFileSize(int index, long size) {
+        if (index >= mExpansionFileSizes.size()) {
+            mExpansionFileSizes.setSize(index + 1);
+        }
+        mExpansionFileSizes.set(index, size);
+    }
+
+    /**
+     * {@inheritDoc} This implementation allows access if either:<br>
+     * <ol>
+     * <li>a LICENSED response was received within the validity period
+     * <li>a RETRY response was received in the last minute, and we are under
+     * the RETRY count or in the RETRY period.
+     * </ol>
+     */
+    public boolean allowAccess() {
+        long ts = System.currentTimeMillis();
+        if (mLastResponse == Policy.LICENSED) {
+            // Check if the LICENSED response occurred within the validity
+            // timeout.
+            if (ts <= mValidityTimestamp) {
+                // Cached LICENSED response is still valid.
+                return true;
+            }
+        } else if (mLastResponse == Policy.RETRY &&
+                ts < mLastResponseTime + MILLIS_PER_MINUTE) {
+            // Only allow access if we are within the retry period or we haven't
+            // used up our
+            // max retries.
+            return (ts <= mRetryUntil || mRetryCount <= mMaxRetries);
+        }
+        return false;
+    }
+
+    private Map<String, String> decodeExtras(String extras) {
+        Map<String, String> results = new HashMap<String, String>();
+        try {
+            URI rawExtras = new URI("?" + extras);
+            List<NameValuePair> extraList = URLEncodedUtils.parse(rawExtras, "UTF-8");
+            for (NameValuePair item : extraList) {
+                String name = item.getName();
+                int i = 0;
+                while (results.containsKey(name)) {
+                    name = item.getName() + ++i;
+                }
+                results.put(name, item.getValue());
+            }
+        } catch (URISyntaxException e) {
+            Log.w(TAG, "Invalid syntax error while decoding extras data from server.");
+        }
+        return results;
+    }
+
+}
diff --git a/Engine/Build/Android/Java/src/com/android/vending/licensing/DeviceLimiter.java b/Engine/Build/Android/Java/src/com/android/vending/licensing/DeviceLimiter.java
new file mode 100644
--- /dev/null
+++ b/Engine/Build/Android/Java/src/com/android/vending/licensing/DeviceLimiter.java
@@ -0,0 +1,47 @@
+/*
+ * Copyright (C) 2010 The Android Open Source Project
+ *
+ * Licensed under the Apache License, Version 2.0 (the "License");
+ * you may not use this file except in compliance with the License.
+ * You may obtain a copy of the License at
+ *
+ *      http://www.apache.org/licenses/LICENSE-2.0
+ *
+ * Unless required by applicable law or agreed to in writing, software
+ * distributed under the License is distributed on an "AS IS" BASIS,
+ * WITHOUT WARRANTIES OR CONDITIONS OF ANY KIND, either express or implied.
+ * See the License for the specific language governing permissions and
+ * limitations under the License.
+ */
+
+package com.google.android.vending.licensing;
+
+/**
+ * Allows the developer to limit the number of devices using a single license.
+ * <p>
+ * The LICENSED response from the server contains a user identifier unique to
+ * the &lt;application, user&gt; pair. The developer can send this identifier
+ * to their own server along with some device identifier (a random number
+ * generated and stored once per application installation,
+ * {@link android.telephony.TelephonyManager#getDeviceId getDeviceId},
+ * {@link android.provider.Settings.Secure#ANDROID_ID ANDROID_ID}, etc).
+ * The more sources used to identify the device, the harder it will be for an
+ * attacker to spoof.
+ * <p>
+ * The server can look at the &lt;application, user, device id&gt; tuple and
+ * restrict a user's application license to run on at most 10 different devices
+ * in a week (for example). We recommend not being too restrictive because a
+ * user might legitimately have multiple devices or be in the process of
+ * changing phones. This will catch egregious violations of multiple people
+ * sharing one license.
+ */
+public interface DeviceLimiter {
+
+    /**
+     * Checks if this device is allowed to use the given user's license.
+     *
+     * @param userId the user whose license the server responded with
+     * @return LICENSED if the device is allowed, NOT_LICENSED if not, RETRY if an error occurs
+     */
+    int isDeviceAllowed(String userId);
+}
diff --git a/Engine/Build/Android/Java/src/com/android/vending/licensing/LicenseChecker.java b/Engine/Build/Android/Java/src/com/android/vending/licensing/LicenseChecker.java
new file mode 100644
--- /dev/null
+++ b/Engine/Build/Android/Java/src/com/android/vending/licensing/LicenseChecker.java
@@ -0,0 +1,351 @@
+/*
+ * Copyright (C) 2010 The Android Open Source Project
+ *
+ * Licensed under the Apache License, Version 2.0 (the "License");
+ * you may not use this file except in compliance with the License.
+ * You may obtain a copy of the License at
+ *
+ *      http://www.apache.org/licenses/LICENSE-2.0
+ *
+ * Unless required by applicable law or agreed to in writing, software
+ * distributed under the License is distributed on an "AS IS" BASIS,
+ * WITHOUT WARRANTIES OR CONDITIONS OF ANY KIND, either express or implied.
+ * See the License for the specific language governing permissions and
+ * limitations under the License.
+ */
+
+package com.google.android.vending.licensing;
+
+import com.google.android.vending.licensing.util.Base64;
+import com.google.android.vending.licensing.util.Base64DecoderException;
+
+import android.content.ComponentName;
+import android.content.Context;
+import android.content.Intent;
+import android.content.ServiceConnection;
+import android.content.pm.PackageManager.NameNotFoundException;
+import android.os.Handler;
+import android.os.HandlerThread;
+import android.os.IBinder;
+import android.os.RemoteException;
+import android.provider.Settings.Secure;
+import android.util.Log;
+
+import java.security.KeyFactory;
+import java.security.NoSuchAlgorithmException;
+import java.security.PublicKey;
+import java.security.SecureRandom;
+import java.security.spec.InvalidKeySpecException;
+import java.security.spec.X509EncodedKeySpec;
+import java.util.Date;
+import java.util.HashSet;
+import java.util.LinkedList;
+import java.util.Queue;
+import java.util.Set;
+
+/**
+ * Client library for Android Market license verifications.
+ * <p>
+ * The LicenseChecker is configured via a {@link Policy} which contains the
+ * logic to determine whether a user should have access to the application. For
+ * example, the Policy can define a threshold for allowable number of server or
+ * client failures before the library reports the user as not having access.
+ * <p>
+ * Must also provide the Base64-encoded RSA public key associated with your
+ * developer account. The public key is obtainable from the publisher site.
+ */
+public class LicenseChecker implements ServiceConnection {
+    private static final String TAG = "LicenseChecker";
+
+    private static final String KEY_FACTORY_ALGORITHM = "RSA";
+
+    // Timeout value (in milliseconds) for calls to service.
+    private static final int TIMEOUT_MS = 10 * 1000;
+
+    private static final SecureRandom RANDOM = new SecureRandom();
+    private static final boolean DEBUG_LICENSE_ERROR = false;
+
+    private ILicensingService mService;
+
+    private PublicKey mPublicKey;
+    private final Context mContext;
+    private final Policy mPolicy;
+    /**
+     * A handler for running tasks on a background thread. We don't want license
+     * processing to block the UI thread.
+     */
+    private Handler mHandler;
+    private final String mPackageName;
+    private final String mVersionCode;
+    private final Set<LicenseValidator> mChecksInProgress = new HashSet<LicenseValidator>();
+    private final Queue<LicenseValidator> mPendingChecks = new LinkedList<LicenseValidator>();
+
+    /**
+     * @param context a Context
+     * @param policy implementation of Policy
+     * @param encodedPublicKey Base64-encoded RSA public key
+     * @throws IllegalArgumentException if encodedPublicKey is invalid
+     */
+    public LicenseChecker(Context context, Policy policy, String encodedPublicKey) {
+        mContext = context;
+        mPolicy = policy;
+        mPublicKey = generatePublicKey(encodedPublicKey);
+        mPackageName = mContext.getPackageName();
+        mVersionCode = getVersionCode(context, mPackageName);
+        HandlerThread handlerThread = new HandlerThread("background thread");
+        handlerThread.start();
+        mHandler = new Handler(handlerThread.getLooper());
+    }
+
+    /**
+     * Generates a PublicKey instance from a string containing the
+     * Base64-encoded public key.
+     * 
+     * @param encodedPublicKey Base64-encoded public key
+     * @throws IllegalArgumentException if encodedPublicKey is invalid
+     */
+    private static PublicKey generatePublicKey(String encodedPublicKey) {
+        try {
+            byte[] decodedKey = Base64.decode(encodedPublicKey);
+            KeyFactory keyFactory = KeyFactory.getInstance(KEY_FACTORY_ALGORITHM);
+
+            return keyFactory.generatePublic(new X509EncodedKeySpec(decodedKey));
+        } catch (NoSuchAlgorithmException e) {
+            // This won't happen in an Android-compatible environment.
+            throw new RuntimeException(e);
+        } catch (Base64DecoderException e) {
+            Log.e(TAG, "Could not decode from Base64.");
+            throw new IllegalArgumentException(e);
+        } catch (InvalidKeySpecException e) {
+            Log.e(TAG, "Invalid key specification.");
+            throw new IllegalArgumentException(e);
+        }
+    }
+
+    /**
+     * Checks if the user should have access to the app.  Binds the service if necessary.
+     * <p>
+     * NOTE: This call uses a trivially obfuscated string (base64-encoded).  For best security,
+     * we recommend obfuscating the string that is passed into bindService using another method
+     * of your own devising.
+     * <p>
+     * source string: "com.android.vending.licensing.ILicensingService"
+     * <p>
+     * @param callback
+     */
+    public synchronized void checkAccess(LicenseCheckerCallback callback) {
+        // If we have a valid recent LICENSED response, we can skip asking
+        // Market.
+        if (mPolicy.allowAccess()) {
+            Log.i(TAG, "Using cached license response");
+            callback.allow(Policy.LICENSED);
+        } else {
+            LicenseValidator validator = new LicenseValidator(mPolicy, new NullDeviceLimiter(),
+                    callback, generateNonce(), mPackageName, mVersionCode);
+
+            if (mService == null) {
+                Log.i(TAG, "Binding to licensing service.");
+                try {
+                    boolean bindResult = mContext
+                            .bindService(
+                                    new Intent(
+                                            new String(
+                                                    Base64.decode("Y29tLmFuZHJvaWQudmVuZGluZy5saWNlbnNpbmcuSUxpY2Vuc2luZ1NlcnZpY2U="))),
+                                    this, // ServiceConnection.
+                                    Context.BIND_AUTO_CREATE);
+
+                    if (bindResult) {
+                        mPendingChecks.offer(validator);
+                    } else {
+                        Log.e(TAG, "Could not bind to service.");
+                        handleServiceConnectionError(validator);
+                    }
+                } catch (SecurityException e) {
+                    callback.applicationError(LicenseCheckerCallback.ERROR_MISSING_PERMISSION);
+                } catch (Base64DecoderException e) {
+                    e.printStackTrace();
+                }
+            } else {
+                mPendingChecks.offer(validator);
+                runChecks();
+            }
+        }
+    }
+
+    private void runChecks() {
+        LicenseValidator validator;
+        while ((validator = mPendingChecks.poll()) != null) {
+            try {
+                Log.i(TAG, "Calling checkLicense on service for " + validator.getPackageName());
+                mService.checkLicense(
+                        validator.getNonce(), validator.getPackageName(),
+                        new ResultListener(validator));
+                mChecksInProgress.add(validator);
+            } catch (RemoteException e) {
+                Log.w(TAG, "RemoteException in checkLicense call.", e);
+                handleServiceConnectionError(validator);
+            }
+        }
+    }
+
+    private synchronized void finishCheck(LicenseValidator validator) {
+        mChecksInProgress.remove(validator);
+        if (mChecksInProgress.isEmpty()) {
+            cleanupService();
+        }
+    }
+
+    private class ResultListener extends ILicenseResultListener.Stub {
+        private final LicenseValidator mValidator;
+        private Runnable mOnTimeout;
+
+        public ResultListener(LicenseValidator validator) {
+            mValidator = validator;
+            mOnTimeout = new Runnable() {
+                public void run() {
+                    Log.i(TAG, "Check timed out.");
+                    handleServiceConnectionError(mValidator);
+                    finishCheck(mValidator);
+                }
+            };
+            startTimeout();
+        }
+
+        private static final int ERROR_CONTACTING_SERVER = 0x101;
+        private static final int ERROR_INVALID_PACKAGE_NAME = 0x102;
+        private static final int ERROR_NON_MATCHING_UID = 0x103;
+
+        // Runs in IPC thread pool. Post it to the Handler, so we can guarantee
+        // either this or the timeout runs.
+        public void verifyLicense(final int responseCode, final String signedData,
+                final String signature) {
+            mHandler.post(new Runnable() {
+                public void run() {
+                    Log.i(TAG, "Received response.");
+                    // Make sure it hasn't already timed out.
+                    if (mChecksInProgress.contains(mValidator)) {
+                        clearTimeout();
+                        mValidator.verify(mPublicKey, responseCode, signedData, signature);
+                        finishCheck(mValidator);
+                    }
+                    if (DEBUG_LICENSE_ERROR) {
+                        boolean logResponse;
+                        String stringError = null;
+                        switch (responseCode) {
+                            case ERROR_CONTACTING_SERVER:
+                                logResponse = true;
+                                stringError = "ERROR_CONTACTING_SERVER";
+                                break;
+                            case ERROR_INVALID_PACKAGE_NAME:
+                                logResponse = true;
+                                stringError = "ERROR_INVALID_PACKAGE_NAME";
+                                break;
+                            case ERROR_NON_MATCHING_UID:
+                                logResponse = true;
+                                stringError = "ERROR_NON_MATCHING_UID";
+                                break;
+                            default:
+                                logResponse = false;
+                        }
+
+                        if (logResponse) {
+                            String android_id = Secure.getString(mContext.getContentResolver(),
+                                    Secure.ANDROID_ID);
+                            Date date = new Date();
+                            Log.d(TAG, "Server Failure: " + stringError);
+                            Log.d(TAG, "Android ID: " + android_id);
+                            Log.d(TAG, "Time: " + date.toGMTString());
+                        }
+                    }
+
+                }
+            });
+        }
+
+        private void startTimeout() {
+            Log.i(TAG, "Start monitoring timeout.");
+            mHandler.postDelayed(mOnTimeout, TIMEOUT_MS);
+        }
+
+        private void clearTimeout() {
+            Log.i(TAG, "Clearing timeout.");
+            mHandler.removeCallbacks(mOnTimeout);
+        }
+    }
+
+    public synchronized void onServiceConnected(ComponentName name, IBinder service) {
+        mService = ILicensingService.Stub.asInterface(service);
+        runChecks();
+    }
+
+    public synchronized void onServiceDisconnected(ComponentName name) {
+        // Called when the connection with the service has been
+        // unexpectedly disconnected. That is, Market crashed.
+        // If there are any checks in progress, the timeouts will handle them.
+        Log.w(TAG, "Service unexpectedly disconnected.");
+        mService = null;
+    }
+
+    /**
+     * Generates policy response for service connection errors, as a result of
+     * disconnections or timeouts.
+     */
+    private synchronized void handleServiceConnectionError(LicenseValidator validator) {
+        mPolicy.processServerResponse(Policy.RETRY, null);
+
+        if (mPolicy.allowAccess()) {
+            validator.getCallback().allow(Policy.RETRY);
+        } else {
+            validator.getCallback().dontAllow(Policy.RETRY);
+        }
+    }
+
+    /** Unbinds service if necessary and removes reference to it. */
+    private void cleanupService() {
+        if (mService != null) {
+            try {
+                mContext.unbindService(this);
+            } catch (IllegalArgumentException e) {
+                // Somehow we've already been unbound. This is a non-fatal
+                // error.
+                Log.e(TAG, "Unable to unbind from licensing service (already unbound)");
+            }
+            mService = null;
+        }
+    }
+
+    /**
+     * Inform the library that the context is about to be destroyed, so that any
+     * open connections can be cleaned up.
+     * <p>
+     * Failure to call this method can result in a crash under certain
+     * circumstances, such as during screen rotation if an Activity requests the
+     * license check or when the user exits the application.
+     */
+    public synchronized void onDestroy() {
+        cleanupService();
+        mHandler.getLooper().quit();
+    }
+
+    /** Generates a nonce (number used once). */
+    private int generateNonce() {
+        return RANDOM.nextInt();
+    }
+
+    /**
+     * Get version code for the application package name.
+     * 
+     * @param context
+     * @param packageName application package name
+     * @return the version code or empty string if package not found
+     */
+    private static String getVersionCode(Context context, String packageName) {
+        try {
+            return String.valueOf(context.getPackageManager().getPackageInfo(packageName, 0).
+                    versionCode);
+        } catch (NameNotFoundException e) {
+            Log.e(TAG, "Package not found. could not get version code.");
+            return "";
+        }
+    }
+}
diff --git a/Engine/Build/Android/Java/src/com/android/vending/licensing/LicenseCheckerCallback.java b/Engine/Build/Android/Java/src/com/android/vending/licensing/LicenseCheckerCallback.java
new file mode 100644
--- /dev/null
+++ b/Engine/Build/Android/Java/src/com/android/vending/licensing/LicenseCheckerCallback.java
@@ -0,0 +1,67 @@
+/*
+ * Copyright (C) 2010 The Android Open Source Project
+ *
+ * Licensed under the Apache License, Version 2.0 (the "License");
+ * you may not use this file except in compliance with the License.
+ * You may obtain a copy of the License at
+ *
+ *      http://www.apache.org/licenses/LICENSE-2.0
+ *
+ * Unless required by applicable law or agreed to in writing, software
+ * distributed under the License is distributed on an "AS IS" BASIS,
+ * WITHOUT WARRANTIES OR CONDITIONS OF ANY KIND, either express or implied.
+ * See the License for the specific language governing permissions and
+ * limitations under the License.
+ */
+
+package com.google.android.vending.licensing;
+
+/**
+ * Callback for the license checker library.
+ * <p>
+ * Upon checking with the Market server and conferring with the {@link Policy},
+ * the library calls the appropriate callback method to communicate the result.
+ * <p>
+ * <b>The callback does not occur in the original checking thread.</b> Your
+ * application should post to the appropriate handling thread or lock
+ * accordingly.
+ * <p>
+ * The reason that is passed back with allow/dontAllow is the base status handed
+ * to the policy for allowed/disallowing the license. Policy.RETRY will call
+ * allow or dontAllow depending on other statistics associated with the policy,
+ * while in most cases Policy.NOT_LICENSED will call dontAllow and
+ * Policy.LICENSED will Allow.
+ */
+public interface LicenseCheckerCallback {
+
+    /**
+     * Allow use. App should proceed as normal.
+     * 
+     * @param reason Policy.LICENSED or Policy.RETRY typically. (although in
+     *            theory the policy can return Policy.NOT_LICENSED here as well)
+     */
+    public void allow(int reason);
+
+    /**
+     * Don't allow use. App should inform user and take appropriate action.
+     * 
+     * @param reason Policy.NOT_LICENSED or Policy.RETRY. (although in theory
+     *            the policy can return Policy.LICENSED here as well ---
+     *            perhaps the call to the LVL took too long, for example)
+     */
+    public void dontAllow(int reason);
+
+    /** Application error codes. */
+    public static final int ERROR_INVALID_PACKAGE_NAME = 1;
+    public static final int ERROR_NON_MATCHING_UID = 2;
+    public static final int ERROR_NOT_MARKET_MANAGED = 3;
+    public static final int ERROR_CHECK_IN_PROGRESS = 4;
+    public static final int ERROR_INVALID_PUBLIC_KEY = 5;
+    public static final int ERROR_MISSING_PERMISSION = 6;
+
+    /**
+     * Error in application code. Caller did not call or set up license checker
+     * correctly. Should be considered fatal.
+     */
+    public void applicationError(int errorCode);
+}
diff --git a/Engine/Build/Android/Java/src/com/android/vending/licensing/LicenseValidator.java b/Engine/Build/Android/Java/src/com/android/vending/licensing/LicenseValidator.java
new file mode 100644
--- /dev/null
+++ b/Engine/Build/Android/Java/src/com/android/vending/licensing/LicenseValidator.java
@@ -0,0 +1,224 @@
+/*
+ * Copyright (C) 2010 The Android Open Source Project
+ *
+ * Licensed under the Apache License, Version 2.0 (the "License");
+ * you may not use this file except in compliance with the License.
+ * You may obtain a copy of the License at
+ *
+ *      http://www.apache.org/licenses/LICENSE-2.0
+ *
+ * Unless required by applicable law or agreed to in writing, software
+ * distributed under the License is distributed on an "AS IS" BASIS,
+ * WITHOUT WARRANTIES OR CONDITIONS OF ANY KIND, either express or implied.
+ * See the License for the specific language governing permissions and
+ * limitations under the License.
+ */
+
+package com.google.android.vending.licensing;
+
+import com.google.android.vending.licensing.util.Base64;
+import com.google.android.vending.licensing.util.Base64DecoderException;
+
+import android.text.TextUtils;
+import android.util.Log;
+
+import java.security.InvalidKeyException;
+import java.security.NoSuchAlgorithmException;
+import java.security.PublicKey;
+import java.security.Signature;
+import java.security.SignatureException;
+
+/**
+ * Contains data related to a licensing request and methods to verify
+ * and process the response.
+ */
+class LicenseValidator {
+    private static final String TAG = "LicenseValidator";
+
+    // Server response codes.
+    private static final int LICENSED = 0x0;
+    private static final int NOT_LICENSED = 0x1;
+    private static final int LICENSED_OLD_KEY = 0x2;
+    private static final int ERROR_NOT_MARKET_MANAGED = 0x3;
+    private static final int ERROR_SERVER_FAILURE = 0x4;
+    private static final int ERROR_OVER_QUOTA = 0x5;
+
+    private static final int ERROR_CONTACTING_SERVER = 0x101;
+    private static final int ERROR_INVALID_PACKAGE_NAME = 0x102;
+    private static final int ERROR_NON_MATCHING_UID = 0x103;
+
+    private final Policy mPolicy;
+    private final LicenseCheckerCallback mCallback;
+    private final int mNonce;
+    private final String mPackageName;
+    private final String mVersionCode;
+    private final DeviceLimiter mDeviceLimiter;
+
+    LicenseValidator(Policy policy, DeviceLimiter deviceLimiter, LicenseCheckerCallback callback,
+             int nonce, String packageName, String versionCode) {
+        mPolicy = policy;
+        mDeviceLimiter = deviceLimiter;
+        mCallback = callback;
+        mNonce = nonce;
+        mPackageName = packageName;
+        mVersionCode = versionCode;
+    }
+
+    public LicenseCheckerCallback getCallback() {
+        return mCallback;
+    }
+
+    public int getNonce() {
+        return mNonce;
+    }
+
+    public String getPackageName() {
+        return mPackageName;
+    }
+
+    private static final String SIGNATURE_ALGORITHM = "SHA1withRSA";
+
+    /**
+     * Verifies the response from server and calls appropriate callback method.
+     *
+     * @param publicKey public key associated with the developer account
+     * @param responseCode server response code
+     * @param signedData signed data from server
+     * @param signature server signature
+     */
+    public void verify(PublicKey publicKey, int responseCode, String signedData, String signature) {
+        String userId = null;
+        // Skip signature check for unsuccessful requests
+        ResponseData data = null;
+        if (responseCode == LICENSED || responseCode == NOT_LICENSED ||
+                responseCode == LICENSED_OLD_KEY) {
+            // Verify signature.
+            try {
+                Signature sig = Signature.getInstance(SIGNATURE_ALGORITHM);
+                sig.initVerify(publicKey);
+                sig.update(signedData.getBytes());
+
+                if (!sig.verify(Base64.decode(signature))) {
+                    Log.e(TAG, "Signature verification failed.");
+                    handleInvalidResponse();
+                    return;
+                }
+            } catch (NoSuchAlgorithmException e) {
+                // This can't happen on an Android compatible device.
+                throw new RuntimeException(e);
+            } catch (InvalidKeyException e) {
+                handleApplicationError(LicenseCheckerCallback.ERROR_INVALID_PUBLIC_KEY);
+                return;
+            } catch (SignatureException e) {
+                throw new RuntimeException(e);
+            } catch (Base64DecoderException e) {
+                Log.e(TAG, "Could not Base64-decode signature.");
+                handleInvalidResponse();
+                return;
+            }
+
+            // Parse and validate response.
+            try {
+                data = ResponseData.parse(signedData);
+            } catch (IllegalArgumentException e) {
+                Log.e(TAG, "Could not parse response.");
+                handleInvalidResponse();
+                return;
+            }
+
+            if (data.responseCode != responseCode) {
+                Log.e(TAG, "Response codes don't match.");
+                handleInvalidResponse();
+                return;
+            }
+
+            if (data.nonce != mNonce) {
+                Log.e(TAG, "Nonce doesn't match.");
+                handleInvalidResponse();
+                return;
+            }
+
+            if (!data.packageName.equals(mPackageName)) {
+                Log.e(TAG, "Package name doesn't match.");
+                handleInvalidResponse();
+                return;
+            }
+
+            if (!data.versionCode.equals(mVersionCode)) {
+                Log.e(TAG, "Version codes don't match.");
+                handleInvalidResponse();
+                return;
+            }
+
+            // Application-specific user identifier.
+            userId = data.userId;
+            if (TextUtils.isEmpty(userId)) {
+                Log.e(TAG, "User identifier is empty.");
+                handleInvalidResponse();
+                return;
+            }
+        }
+
+        switch (responseCode) {
+            case LICENSED:
+            case LICENSED_OLD_KEY:
+                int limiterResponse = mDeviceLimiter.isDeviceAllowed(userId);
+                handleResponse(limiterResponse, data);
+                break;
+            case NOT_LICENSED:
+                handleResponse(Policy.NOT_LICENSED, data);
+                break;
+            case ERROR_CONTACTING_SERVER:
+                Log.w(TAG, "Error contacting licensing server.");
+                handleResponse(Policy.RETRY, data);
+                break;
+            case ERROR_SERVER_FAILURE:
+                Log.w(TAG, "An error has occurred on the licensing server.");
+                handleResponse(Policy.RETRY, data);
+                break;
+            case ERROR_OVER_QUOTA:
+                Log.w(TAG, "Licensing server is refusing to talk to this device, over quota.");
+                handleResponse(Policy.RETRY, data);
+                break;
+            case ERROR_INVALID_PACKAGE_NAME:
+                handleApplicationError(LicenseCheckerCallback.ERROR_INVALID_PACKAGE_NAME);
+                break;
+            case ERROR_NON_MATCHING_UID:
+                handleApplicationError(LicenseCheckerCallback.ERROR_NON_MATCHING_UID);
+                break;
+            case ERROR_NOT_MARKET_MANAGED:
+                handleApplicationError(LicenseCheckerCallback.ERROR_NOT_MARKET_MANAGED);
+                break;
+            default:
+                Log.e(TAG, "Unknown response code for license check.");
+                handleInvalidResponse();
+        }
+    }
+
+    /**
+     * Confers with policy and calls appropriate callback method.
+     *
+     * @param response
+     * @param rawData
+     */
+    private void handleResponse(int response, ResponseData rawData) {
+        // Update policy data and increment retry counter (if needed)
+        mPolicy.processServerResponse(response, rawData);
+
+        // Given everything we know, including cached data, ask the policy if we should grant
+        // access.
+        if (mPolicy.allowAccess()) {
+            mCallback.allow(response);
+        } else {
+            mCallback.dontAllow(response);
+        }
+    }
+
+    private void handleApplicationError(int code) {
+        mCallback.applicationError(code);
+    }
+
+    private void handleInvalidResponse() {
+        mCallback.dontAllow(Policy.NOT_LICENSED);
+    }
+}
diff --git a/Engine/Build/Android/Java/src/com/android/vending/licensing/NullDeviceLimiter.java b/Engine/Build/Android/Java/src/com/android/vending/licensing/NullDeviceLimiter.java
new file mode 100644
--- /dev/null
+++ b/Engine/Build/Android/Java/src/com/android/vending/licensing/NullDeviceLimiter.java
@@ -0,0 +1,32 @@
+/*
+ * Copyright (C) 2010 The Android Open Source Project
+ *
+ * Licensed under the Apache License, Version 2.0 (the "License");
+ * you may not use this file except in compliance with the License.
+ * You may obtain a copy of the License at
+ *
+ *      http://www.apache.org/licenses/LICENSE-2.0
+ *
+ * Unless required by applicable law or agreed to in writing, software
+ * distributed under the License is distributed on an "AS IS" BASIS,
+ * WITHOUT WARRANTIES OR CONDITIONS OF ANY KIND, either express or implied.
+ * See the License for the specific language governing permissions and
+ * limitations under the License.
+ */
+
+package com.google.android.vending.licensing;
+
+/**
+ * A DeviceLimiter that doesn't limit the number of devices that can use a
+ * given user's license.
+ * <p>
+ * Unless you have reason to believe that your application is being pirated
+ * by multiple users using the same license (signing in to Market as the same
+ * user), we recommend you use this implementation.
+ */
+public class NullDeviceLimiter implements DeviceLimiter {
+
+    public int isDeviceAllowed(String userId) {
+        return Policy.LICENSED;
+    }
+}
diff --git a/Engine/Build/Android/Java/src/com/android/vending/licensing/Obfuscator.java b/Engine/Build/Android/Java/src/com/android/vending/licensing/Obfuscator.java
new file mode 100644
--- /dev/null
+++ b/Engine/Build/Android/Java/src/com/android/vending/licensing/Obfuscator.java
@@ -0,0 +1,48 @@
+/*
+ * Copyright (C) 2010 The Android Open Source Project
+ *
+ * Licensed under the Apache License, Version 2.0 (the "License");
+ * you may not use this file except in compliance with the License.
+ * You may obtain a copy of the License at
+ *
+ *      http://www.apache.org/licenses/LICENSE-2.0
+ *
+ * Unless required by applicable law or agreed to in writing, software
+ * distributed under the License is distributed on an "AS IS" BASIS,
+ * WITHOUT WARRANTIES OR CONDITIONS OF ANY KIND, either express or implied.
+ * See the License for the specific language governing permissions and
+ * limitations under the License.
+ */
+
+package com.google.android.vending.licensing;
+
+/**
+ * Interface used as part of a {@link Policy} to allow application authors to obfuscate
+ * licensing data that will be stored into a SharedPreferences file.
+ * <p>
+ * Any transformation scheme must be reversable. Implementing classes may optionally implement an
+ * integrity check to further prevent modification to preference data. Implementing classes
+ * should use device-specific information as a key in the obfuscation algorithm to prevent
+ * obfuscated preferences from being shared among devices.
+ */
+public interface Obfuscator {
+
+    /**
+     * Obfuscate a string that is being stored into shared preferences.
+     *
+     * @param original The data that is to be obfuscated.
+     * @param key The key for the data that is to be obfuscated.
+     * @return A transformed version of the original data.
+     */
+    String obfuscate(String original, String key);
+
+    /**
+     * Undo the transformation applied to data by the obfuscate() method.
+     *
+     * @param original The data that is to be obfuscated.
+     * @param key The key for the data that is to be obfuscated.
+     * @return A transformed version of the original data.
+     * @throws ValidationException Optionally thrown if a data integrity check fails.
+     */
+    String unobfuscate(String obfuscated, String key) throws ValidationException;
+}
diff --git a/Engine/Build/Android/Java/src/com/android/vending/licensing/Policy.java b/Engine/Build/Android/Java/src/com/android/vending/licensing/Policy.java
new file mode 100644
--- /dev/null
+++ b/Engine/Build/Android/Java/src/com/android/vending/licensing/Policy.java
@@ -0,0 +1,59 @@
+/*
+ * Copyright (C) 2010 The Android Open Source Project
+ *
+ * Licensed under the Apache License, Version 2.0 (the "License");
+ * you may not use this file except in compliance with the License.
+ * You may obtain a copy of the License at
+ *
+ *      http://www.apache.org/licenses/LICENSE-2.0
+ *
+ * Unless required by applicable law or agreed to in writing, software
+ * distributed under the License is distributed on an "AS IS" BASIS,
+ * WITHOUT WARRANTIES OR CONDITIONS OF ANY KIND, either express or implied.
+ * See the License for the specific language governing permissions and
+ * limitations under the License.
+ */
+
+package com.google.android.vending.licensing;
+
+/**
+ * Policy used by {@link LicenseChecker} to determine whether a user should have
+ * access to the application.
+ */
+public interface Policy {
+
+    /**
+     * Change these values to make it more difficult for tools to automatically
+     * strip LVL protection from your APK.
+     */
+
+    /**
+     * LICENSED means that the server returned back a valid license response
+     */
+    public static final int LICENSED = 0x0100;
+    /**
+     * NOT_LICENSED means that the server returned back a valid license response
+     * that indicated that the user definitively is not licensed
+     */
+    public static final int NOT_LICENSED = 0x0231;
+    /**
+     * RETRY means that the license response was unable to be determined ---
+     * perhaps as a result of faulty networking
+     */
+    public static final int RETRY = 0x0123;
+
+    /**
+     * Provide results from contact with the license server. Retry counts are
+     * incremented if the current value of response is RETRY. Results will be
+     * used for any future policy decisions.
+     * 
+     * @param response the result from validating the server response
+     * @param rawData the raw server response data, can be null for RETRY
+     */
+    void processServerResponse(int response, ResponseData rawData);
+
+    /**
+     * Check if the user should be allowed access to the application.
+     */
+    boolean allowAccess();
+}
diff --git a/Engine/Build/Android/Java/src/com/android/vending/licensing/PreferenceObfuscator.java b/Engine/Build/Android/Java/src/com/android/vending/licensing/PreferenceObfuscator.java
new file mode 100644
--- /dev/null
+++ b/Engine/Build/Android/Java/src/com/android/vending/licensing/PreferenceObfuscator.java
@@ -0,0 +1,77 @@
+/*
+ * Copyright (C) 2010 The Android Open Source Project
+ *
+ * Licensed under the Apache License, Version 2.0 (the "License");
+ * you may not use this file except in compliance with the License.
+ * You may obtain a copy of the License at
+ *
+ *      http://www.apache.org/licenses/LICENSE-2.0
+ *
+ * Unless required by applicable law or agreed to in writing, software
+ * distributed under the License is distributed on an "AS IS" BASIS,
+ * WITHOUT WARRANTIES OR CONDITIONS OF ANY KIND, either express or implied.
+ * See the License for the specific language governing permissions and
+ * limitations under the License.
+ */
+
+package com.google.android.vending.licensing;
+
+import android.content.SharedPreferences;
+import android.util.Log;
+
+/**
+ * An wrapper for SharedPreferences that transparently performs data obfuscation.
+ */
+public class PreferenceObfuscator {
+
+    private static final String TAG = "PreferenceObfuscator";
+
+    private final SharedPreferences mPreferences;
+    private final Obfuscator mObfuscator;
+    private SharedPreferences.Editor mEditor;
+
+    /**
+     * Constructor.
+     *
+     * @param sp A SharedPreferences instance provided by the system.
+     * @param o The Obfuscator to use when reading or writing data.
+     */
+    public PreferenceObfuscator(SharedPreferences sp, Obfuscator o) {
+        mPreferences = sp;
+        mObfuscator = o;
+        mEditor = null;
+    }
+
+    public void putString(String key, String value) {
+        if (mEditor == null) {
+            mEditor = mPreferences.edit();
+        }
+        String obfuscatedValue = mObfuscator.obfuscate(value, key);
+        mEditor.putString(key, obfuscatedValue);
+    }
+
+    public String getString(String key, String defValue) {
+        String result;
+        String value = mPreferences.getString(key, null);
+        if (value != null) {
+            try {
+                result = mObfuscator.unobfuscate(value, key);
+            } catch (ValidationException e) {
+                // Unable to unobfuscate, data corrupt or tampered
+                Log.w(TAG, "Validation error while reading preference: " + key);
+                result = defValue;
+            }
+        } else {
+            // Preference not found
+            result = defValue;
+        }
+        return result;
+    }
+
+    public void commit() {
+        if (mEditor != null) {
+            mEditor.commit();
+            mEditor = null;
+        }
+    }
+}
diff --git a/Engine/Build/Android/Java/src/com/android/vending/licensing/ResponseData.java b/Engine/Build/Android/Java/src/com/android/vending/licensing/ResponseData.java
new file mode 100644
--- /dev/null
+++ b/Engine/Build/Android/Java/src/com/android/vending/licensing/ResponseData.java
@@ -0,0 +1,79 @@
+/*
+ * Copyright (C) 2010 The Android Open Source Project
+ *
+ * Licensed under the Apache License, Version 2.0 (the "License");
+ * you may not use this file except in compliance with the License.
+ * You may obtain a copy of the License at
+ *
+ *      http://www.apache.org/licenses/LICENSE-2.0
+ *
+ * Unless required by applicable law or agreed to in writing, software
+ * distributed under the License is distributed on an "AS IS" BASIS,
+ * WITHOUT WARRANTIES OR CONDITIONS OF ANY KIND, either express or implied.
+ * See the License for the specific language governing permissions and
+ * limitations under the License.
+ */
+
+package com.google.android.vending.licensing;
+
+import java.util.regex.Pattern;
+
+import android.text.TextUtils;
+
+/**
+ * ResponseData from licensing server.
+ */
+public class ResponseData {
+
+    public int responseCode;
+    public int nonce;
+    public String packageName;
+    public String versionCode;
+    public String userId;
+    public long timestamp;
+    /** Response-specific data. */
+    public String extra;
+
+    /**
+     * Parses response string into ResponseData.
+     *
+     * @param responseData response data string
+     * @throws IllegalArgumentException upon parsing error
+     * @return ResponseData object
+     */
+    public static ResponseData parse(String responseData) {
+        // Must parse out main response data and response-specific data.
+    	int index = responseData.indexOf(':');
+    	String mainData, extraData;
+    	if ( -1 == index ) {
+    		mainData = responseData;
+    		extraData = "";
+    	} else {
+    		mainData = responseData.substring(0, index);
+    		extraData = index >= responseData.length() ? "" : responseData.substring(index+1);
+    	}
+
+        String [] fields = TextUtils.split(mainData, Pattern.quote("|"));
+        if (fields.length < 6) {
+            throw new IllegalArgumentException("Wrong number of fields.");
+        }
+
+        ResponseData data = new ResponseData();
+        data.extra = extraData;
+        data.responseCode = Integer.parseInt(fields[0]);
+        data.nonce = Integer.parseInt(fields[1]);
+        data.packageName = fields[2];
+        data.versionCode = fields[3];
+        // Application-specific user identifier.
+        data.userId = fields[4];
+        data.timestamp = Long.parseLong(fields[5]);
+
+        return data;
+    }
+
+    @Override
+    public String toString() {
+        return TextUtils.join("|", new Object [] { responseCode, nonce, packageName, versionCode,
+            userId, timestamp });
+    }
+}
diff --git a/Engine/Build/Android/Java/src/com/android/vending/licensing/ServerManagedPolicy.java b/Engine/Build/Android/Java/src/com/android/vending/licensing/ServerManagedPolicy.java
new file mode 100644
--- /dev/null
+++ b/Engine/Build/Android/Java/src/com/android/vending/licensing/ServerManagedPolicy.java
@@ -0,0 +1,276 @@
+/*
+ * Copyright (C) 2010 The Android Open Source Project
+ *
+ * Licensed under the Apache License, Version 2.0 (the "License");
+ * you may not use this file except in compliance with the License.
+ * You may obtain a copy of the License at
+ *
+ *      http://www.apache.org/licenses/LICENSE-2.0
+ *
+ * Unless required by applicable law or agreed to in writing, software
+ * distributed under the License is distributed on an "AS IS" BASIS,
+ * WITHOUT WARRANTIES OR CONDITIONS OF ANY KIND, either express or implied.
+ * See the License for the specific language governing permissions and
+ * limitations under the License.
+ */
+
+package com.google.android.vending.licensing;
+
+import org.apache.http.NameValuePair;
+import org.apache.http.client.utils.URLEncodedUtils;
+
+import java.net.URI;
+import java.net.URISyntaxException;
+import java.util.HashMap;
+import java.util.List;
+import java.util.Map;
+
+import android.content.Context;
+import android.content.SharedPreferences;
+import android.util.Log;
+
+/**
+ * Default policy. All policy decisions are based off of response data received
+ * from the licensing service. Specifically, the licensing server sends the
+ * following information: response validity period, error retry period, and
+ * error retry count.
+ * <p>
+ * These values will vary based on the the way the application is configured in
+ * the Android Market publishing console, such as whether the application is
+ * marked as free or is within its refund period, as well as how often an
+ * application is checking with the licensing service.
+ * <p>
+ * Developers who need more fine grained control over their application's
+ * licensing policy should implement a custom Policy.
+ */
+public class ServerManagedPolicy implements Policy {
+
+    private static final String TAG = "ServerManagedPolicy";
+    private static final String PREFS_FILE = "com.android.vending.licensing.ServerManagedPolicy";
+    private static final String PREF_LAST_RESPONSE = "lastResponse";
+    private static final String PREF_VALIDITY_TIMESTAMP = "validityTimestamp";
+    private static final String PREF_RETRY_UNTIL = "retryUntil";
+    private static final String PREF_MAX_RETRIES = "maxRetries";
+    private static final String PREF_RETRY_COUNT = "retryCount";
+    private static final String DEFAULT_VALIDITY_TIMESTAMP = "0";
+    private static final String DEFAULT_RETRY_UNTIL = "0";
+    private static final String DEFAULT_MAX_RETRIES = "0";
+    private static final String DEFAULT_RETRY_COUNT = "0";
+
+    private static final long MILLIS_PER_MINUTE = 60 * 1000;
+
+    private long mValidityTimestamp;
+    private long mRetryUntil;
+    private long mMaxRetries;
+    private long mRetryCount;
+    private long mLastResponseTime = 0;
+    private int mLastResponse;
+    private PreferenceObfuscator mPreferences;
+
+    /**
+     * @param context The context for the current application
+     * @param obfuscator An obfuscator to be used with preferences.
+     */
+    public ServerManagedPolicy(Context context, Obfuscator obfuscator) {
+        // Import old values
+        SharedPreferences sp = context.getSharedPreferences(PREFS_FILE, Context.MODE_PRIVATE);
+        mPreferences = new PreferenceObfuscator(sp, obfuscator);
+        mLastResponse = Integer.parseInt(
+            mPreferences.getString(PREF_LAST_RESPONSE, Integer.toString(Policy.RETRY)));
+        mValidityTimestamp = Long.parseLong(mPreferences.getString(PREF_VALIDITY_TIMESTAMP,
+                DEFAULT_VALIDITY_TIMESTAMP));
+        mRetryUntil = Long.parseLong(mPreferences.getString(PREF_RETRY_UNTIL, DEFAULT_RETRY_UNTIL));
+        mMaxRetries = Long.parseLong(mPreferences.getString(PREF_MAX_RETRIES, DEFAULT_MAX_RETRIES));
+        mRetryCount = Long.parseLong(mPreferences.getString(PREF_RETRY_COUNT, DEFAULT_RETRY_COUNT));
+    }
+
+    /**
+     * Process a new response from the license server.
+     * <p>
+     * This data will be used for computing future policy decisions. The
+     * following parameters are processed:
+     * <ul>
+     * <li>VT: the timestamp that the client should consider the response
+     *   valid until
+     * <li>GT: the timestamp that the client should ignore retry errors until
+     * <li>GR: the number of retry errors that the client should ignore
+     * </ul>
+     *
+     * @param response the result from validating the server response
+     * @param rawData the raw server response data
+     */
+    public void processServerResponse(int response, ResponseData rawData) {
+
+        // Update retry counter
+        if (response != Policy.RETRY) {
+            setRetryCount(0);
+        } else {
+            setRetryCount(mRetryCount + 1);
+        }
+
+        if (response == Policy.LICENSED) {
+            // Update server policy data
+            Map<String, String> extras = decodeExtras(rawData.extra);
+            mLastResponse = response;
+            setValidityTimestamp(extras.get("VT"));
+            setRetryUntil(extras.get("GT"));
+            setMaxRetries(extras.get("GR"));
+        } else if (response == Policy.NOT_LICENSED) {
+            // Clear out stale policy data
+            setValidityTimestamp(DEFAULT_VALIDITY_TIMESTAMP);
+            setRetryUntil(DEFAULT_RETRY_UNTIL);
+            setMaxRetries(DEFAULT_MAX_RETRIES);
+        }
+
+        setLastResponse(response);
+        mPreferences.commit();
+    }
+
+    /**
+     * Set the last license response received from the server and add to
+     * preferences. You must manually call PreferenceObfuscator.commit() to
+     * commit these changes to disk.
+     *
+     * @param l the response
+     */
+    private void setLastResponse(int l) {
+        mLastResponseTime = System.currentTimeMillis();
+        mLastResponse = l;
+        mPreferences.putString(PREF_LAST_RESPONSE, Integer.toString(l));
+    }
+
+    /**
+     * Set the current retry count and add to preferences. You must manually
+     * call PreferenceObfuscator.commit() to commit these changes to disk.
+     *
+     * @param c the new retry count
+     */
+    private void setRetryCount(long c) {
+        mRetryCount = c;
+        mPreferences.putString(PREF_RETRY_COUNT, Long.toString(c));
+    }
+
+    public long getRetryCount() {
+        return mRetryCount;
+    }
+
+    /**
+     * Set the last validity timestamp (VT) received from the server and add to
+     * preferences. You must manually call PreferenceObfuscator.commit() to
+     * commit these changes to disk.
+     *
+     * @param validityTimestamp the VT string received
+     */
+    private void setValidityTimestamp(String validityTimestamp) {
+        Long lValidityTimestamp;
+        try {
+            lValidityTimestamp = Long.parseLong(validityTimestamp);
+        } catch (NumberFormatException e) {
+            // No response or not parsable, expire in one minute.
+            Log.w(TAG, "License validity timestamp (VT) missing, caching for a minute");
+            lValidityTimestamp = System.currentTimeMillis() + MILLIS_PER_MINUTE;
+            validityTimestamp = Long.toString(lValidityTimestamp);
+        }
+
+        mValidityTimestamp = lValidityTimestamp;
+        mPreferences.putString(PREF_VALIDITY_TIMESTAMP, validityTimestamp);
+    }
+
+    public long getValidityTimestamp() {
+        return mValidityTimestamp;
+    }
+
+    /**
+     * Set the retry until timestamp (GT) received from the server and add to
+     * preferences. You must manually call PreferenceObfuscator.commit() to
+     * commit these changes to disk.
+     *
+     * @param retryUntil the GT string received
+     */
+    private void setRetryUntil(String retryUntil) {
+        Long lRetryUntil;
+        try {
+            lRetryUntil = Long.parseLong(retryUntil);
+        } catch (NumberFormatException e) {
+            // No response or not parsable, expire immediately
+            Log.w(TAG, "License retry timestamp (GT) missing, grace period disabled");
+            retryUntil = "0";
+            lRetryUntil = 0l;
+        }
+
+        mRetryUntil = lRetryUntil;
+        mPreferences.putString(PREF_RETRY_UNTIL, retryUntil);
+    }
+
+    public long getRetryUntil() {
+      return mRetryUntil;
+    }
+
+    /**
+     * Set the max retries value (GR) as received from the server and add to
+     * preferences. You must manually call PreferenceObfuscator.commit() to
+     * commit these changes to disk.
+     *
+     * @param maxRetries the GR string received
+     */
+    private void setMaxRetries(String maxRetries) {
+        Long lMaxRetries;
+        try {
+            lMaxRetries = Long.parseLong(maxRetries);
+        } catch (NumberFormatException e) {
+            // No response or not parsable, expire immediately
+            Log.w(TAG, "Licence retry count (GR) missing, grace period disabled");
+            maxRetries = "0";
+            lMaxRetries = 0l;
+        }
+
+        mMaxRetries = lMaxRetries;
+        mPreferences.putString(PREF_MAX_RETRIES, maxRetries);
+    }
+
+    public long getMaxRetries() {
+        return mMaxRetries;
+    }
+
+    /**
+     * {@inheritDoc}
+     *
+     * This implementation allows access if either:<br>
+     * <ol>
+     * <li>a LICENSED response was received within the validity period
+     * <li>a RETRY response was received in the last minute, and we are under
+     * the RETRY count or in the RETRY period.
+     * </ol>
+     */
+    public boolean allowAccess() {
+        long ts = System.currentTimeMillis();
+        if (mLastResponse == Policy.LICENSED) {
+            // Check if the LICENSED response occurred within the validity timeout.
+            if (ts <= mValidityTimestamp) {
+                // Cached LICENSED response is still valid.
+                return true;
+            }
+        } else if (mLastResponse == Policy.RETRY &&
+                   ts < mLastResponseTime + MILLIS_PER_MINUTE) {
+            // Only allow access if we are within the retry period or we haven't used up our
+            // max retries.
+            return (ts <= mRetryUntil || mRetryCount <= mMaxRetries);
+        }
+        return false;
+    }
+
+    private Map<String, String> decodeExtras(String extras) {
+        Map<String, String> results = new HashMap<String, String>();
+        try {
+            URI rawExtras = new URI("?" + extras);
+            List<NameValuePair> extraList = URLEncodedUtils.parse(rawExtras, "UTF-8");
+            for (NameValuePair item : extraList) {
+                results.put(item.getName(), item.getValue());
+            }
+        } catch (URISyntaxException e) {
+          Log.w(TAG, "Invalid syntax error while decoding extras data from server.");
+        }
+        return results;
+    }
+
+}
diff --git a/Engine/Build/Android/Java/src/com/android/vending/licensing/StrictPolicy.java b/Engine/Build/Android/Java/src/com/android/vending/licensing/StrictPolicy.java
new file mode 100644
--- /dev/null
+++ b/Engine/Build/Android/Java/src/com/android/vending/licensing/StrictPolicy.java
@@ -0,0 +1,63 @@
+/*
+ * Copyright (C) 2010 The Android Open Source Project
+ *
+ * Licensed under the Apache License, Version 2.0 (the "License");
+ * you may not use this file except in compliance with the License.
+ * You may obtain a copy of the License at
+ *
+ *      http://www.apache.org/licenses/LICENSE-2.0
+ *
+ * Unless required by applicable law or agreed to in writing, software
+ * distributed under the License is distributed on an "AS IS" BASIS,
+ * WITHOUT WARRANTIES OR CONDITIONS OF ANY KIND, either express or implied.
+ * See the License for the specific language governing permissions and
+ * limitations under the License.
+ */
+
+package com.google.android.vending.licensing;
+
+/**
+ * Non-caching policy. All requests will be sent to the licensing service,
+ * and no local caching is performed.
+ * <p>
+ * Using a non-caching policy ensures that there is no local preference data
+ * for malicious users to tamper with. As a side effect, applications
+ * will not be permitted to run while offline. Developers should carefully
+ * weigh the risks of using this Policy over one which implements caching,
+ * such as ServerManagedPolicy.
+ * <p>
+ * Access to the application is only allowed if a LICESNED response is.
+ * received. All other responses (including RETRY) will deny access.
+ */
+public class StrictPolicy implements Policy {
+
+    private int mLastResponse;
+
+    public StrictPolicy() {
+        // Set default policy. This will force the application to check the policy on launch.
+        mLastResponse = Policy.RETRY;
+    }
+
+    /**
+     * Process a new response from the license server. Since we aren't
+     * performing any caching, this equates to reading the LicenseResponse.
+     * Any ResponseData provided is ignored.
+     *
+     * @param response the result from validating the server response
+     * @param rawData the raw server response data
+     */
+    public void processServerResponse(int response, ResponseData rawData) {
+        mLastResponse = response;
+    }
+
+    /**
+     * {@inheritDoc}
+     *
+     * This implementation allows access if and only if a LICENSED response
+     * was received the last time the server was contacted.
+     */
+    public boolean allowAccess() {
+        return (mLastResponse == Policy.LICENSED);
+    }
+
+}
diff --git a/Engine/Build/Android/Java/src/com/android/vending/licensing/ValidationException.java b/Engine/Build/Android/Java/src/com/android/vending/licensing/ValidationException.java
new file mode 100644
--- /dev/null
+++ b/Engine/Build/Android/Java/src/com/android/vending/licensing/ValidationException.java
@@ -0,0 +1,33 @@
+/*
+ * Copyright (C) 2010 The Android Open Source Project
+ *
+ * Licensed under the Apache License, Version 2.0 (the "License");
+ * you may not use this file except in compliance with the License.
+ * You may obtain a copy of the License at
+ *
+ *      http://www.apache.org/licenses/LICENSE-2.0
+ *
+ * Unless required by applicable law or agreed to in writing, software
+ * distributed under the License is distributed on an "AS IS" BASIS,
+ * WITHOUT WARRANTIES OR CONDITIONS OF ANY KIND, either express or implied.
+ * See the License for the specific language governing permissions and
+ * limitations under the License.
+ */
+
+package com.google.android.vending.licensing;
+
+/**
+ * Indicates that an error occurred while validating the integrity of data managed by an
+ * {@link Obfuscator}.}
+ */
+public class ValidationException extends Exception {
+    public ValidationException() {
+      super();
+    }
+
+    public ValidationException(String s) {
+      super(s);
+    }
+
+    private static final long serialVersionUID = 1L;
+}
diff --git a/Engine/Build/Android/Java/src/com/android/vending/licensing/util/Base64.java b/Engine/Build/Android/Java/src/com/android/vending/licensing/util/Base64.java
new file mode 100644
--- /dev/null
+++ b/Engine/Build/Android/Java/src/com/android/vending/licensing/util/Base64.java
@@ -0,0 +1,570 @@
+// Portions copyright 2002, Google, Inc.
+//
+// Licensed under the Apache License, Version 2.0 (the "License");
+// you may not use this file except in compliance with the License.
+// You may obtain a copy of the License at
+//
+//     http://www.apache.org/licenses/LICENSE-2.0
+//
+// Unless required by applicable law or agreed to in writing, software
+// distributed under the License is distributed on an "AS IS" BASIS,
+// WITHOUT WARRANTIES OR CONDITIONS OF ANY KIND, either express or implied.
+// See the License for the specific language governing permissions and
+// limitations under the License.
+
+package com.google.android.vending.licensing.util;
+
+// This code was converted from code at http://iharder.sourceforge.net/base64/
+// Lots of extraneous features were removed.
+/* The original code said:
+ * <p>
+ * I am placing this code in the Public Domain. Do with it as you will.
+ * This software comes with no guarantees or warranties but with
+ * plenty of well-wishing instead!
+ * Please visit
+ * <a href="http://iharder.net/xmlizable">http://iharder.net/xmlizable</a>
+ * periodically to check for updates or to contribute improvements.
+ * </p>
+ *
+ * @author Robert Harder
+ * @author [REDACTED_EMAIL]
+ * @version 1.3
+ */
+
+/**
+ * Base64 converter class. This code is not a full-blown MIME encoder;
+ * it simply converts binary data to base64 data and back.
+ *
+ * <p>Note {@link CharBase64} is a GWT-compatible implementation of this
+ * class.
+ */
+public class Base64 {
+  /** Specify encoding (value is {@code true}). */
+  public final static boolean ENCODE = true;
+
+  /** Specify decoding (value is {@code false}). */
+  public final static boolean DECODE = false;
+
+  /** The equals sign (=) as a byte. */
+  private final static byte EQUALS_SIGN = (byte) '=';
+
+  /** The new line character (\n) as a byte. */
+  private final static byte NEW_LINE = (byte) '\n';
+
+  /**
+   * The 64 valid Base64 values.
+   */
+  private final static byte[] ALPHABET =
+      {(byte) 'A', (byte) 'B', (byte) 'C', (byte) 'D', (byte) 'E', (byte) 'F',
+          (byte) 'G', (byte) 'H', (byte) 'I', (byte) 'J', (byte) 'K',
+          (byte) 'L', (byte) 'M', (byte) 'N', (byte) 'O', (byte) 'P',
+          (byte) 'Q', (byte) 'R', (byte) 'S', (byte) 'T', (byte) 'U',
+          (byte) 'V', (byte) 'W', (byte) 'X', (byte) 'Y', (byte) 'Z',
+          (byte) 'a', (byte) 'b', (byte) 'c', (byte) 'd', (byte) 'e',
+          (byte) 'f', (byte) 'g', (byte) 'h', (byte) 'i', (byte) 'j',
+          (byte) 'k', (byte) 'l', (byte) 'm', (byte) 'n', (byte) 'o',
+          (byte) 'p', (byte) 'q', (byte) 'r', (byte) 's', (byte) 't',
+          (byte) 'u', (byte) 'v', (byte) 'w', (byte) 'x', (byte) 'y',
+          (byte) 'z', (byte) '0', (byte) '1', (byte) '2', (byte) '3',
+          (byte) '4', (byte) '5', (byte) '6', (byte) '7', (byte) '8',
+          (byte) '9', (byte) '+', (byte) '/'};
+
+  /**
+   * The 64 valid web safe Base64 values.
+   */
+  private final static byte[] WEBSAFE_ALPHABET =
+      {(byte) 'A', (byte) 'B', (byte) 'C', (byte) 'D', (byte) 'E', (byte) 'F',
+          (byte) 'G', (byte) 'H', (byte) 'I', (byte) 'J', (byte) 'K',
+          (byte) 'L', (byte) 'M', (byte) 'N', (byte) 'O', (byte) 'P',
+          (byte) 'Q', (byte) 'R', (byte) 'S', (byte) 'T', (byte) 'U',
+          (byte) 'V', (byte) 'W', (byte) 'X', (byte) 'Y', (byte) 'Z',
+          (byte) 'a', (byte) 'b', (byte) 'c', (byte) 'd', (byte) 'e',
+          (byte) 'f', (byte) 'g', (byte) 'h', (byte) 'i', (byte) 'j',
+          (byte) 'k', (byte) 'l', (byte) 'm', (byte) 'n', (byte) 'o',
+          (byte) 'p', (byte) 'q', (byte) 'r', (byte) 's', (byte) 't',
+          (byte) 'u', (byte) 'v', (byte) 'w', (byte) 'x', (byte) 'y',
+          (byte) 'z', (byte) '0', (byte) '1', (byte) '2', (byte) '3',
+          (byte) '4', (byte) '5', (byte) '6', (byte) '7', (byte) '8',
+          (byte) '9', (byte) '-', (byte) '_'};
+
+  /**
+   * Translates a Base64 value to either its 6-bit reconstruction value
+   * or a negative number indicating some other meaning.
+   **/
+  private final static byte[] DECODABET = {-9, -9, -9, -9, -9, -9, -9, -9, -9, // Decimal  0 -  8
+      -5, -5, // Whitespace: Tab and Linefeed
+      -9, -9, // Decimal 11 - 12
+      -5, // Whitespace: Carriage Return
+      -9, -9, -9, -9, -9, -9, -9, -9, -9, -9, -9, -9, -9, // Decimal 14 - 26
+      -9, -9, -9, -9, -9, // Decimal 27 - 31
+      -5, // Whitespace: Space
+      -9, -9, -9, -9, -9, -9, -9, -9, -9, -9, // Decimal 33 - 42
+      62, // Plus sign at decimal 43
+      -9, -9, -9, // Decimal 44 - 46
+      63, // Slash at decimal 47
+      52, 53, 54, 55, 56, 57, 58, 59, 60, 61, // Numbers zero through nine
+      -9, -9, -9, // Decimal 58 - 60
+      -1, // Equals sign at decimal 61
+      -9, -9, -9, // Decimal 62 - 64
+      0, 1, 2, 3, 4, 5, 6, 7, 8, 9, 10, 11, 12, 13, // Letters 'A' through 'N'
+      14, 15, 16, 17, 18, 19, 20, 21, 22, 23, 24, 25, // Letters 'O' through 'Z'
+      -9, -9, -9, -9, -9, -9, // Decimal 91 - 96
+      26, 27, 28, 29, 30, 31, 32, 33, 34, 35, 36, 37, 38, // Letters 'a' through 'm'
+      39, 40, 41, 42, 43, 44, 45, 46, 47, 48, 49, 50, 51, // Letters 'n' through 'z'
+      -9, -9, -9, -9, -9 // Decimal 123 - 127
+      /*  ,-9,-9,-9,-9,-9,-9,-9,-9,-9,-9,-9,-9,     // Decimal 128 - 139
+        -9,-9,-9,-9,-9,-9,-9,-9,-9,-9,-9,-9,-9,     // Decimal 140 - 152
+        -9,-9,-9,-9,-9,-9,-9,-9,-9,-9,-9,-9,-9,     // Decimal 153 - 165
+        -9,-9,-9,-9,-9,-9,-9,-9,-9,-9,-9,-9,-9,     // Decimal 166 - 178
+        -9,-9,-9,-9,-9,-9,-9,-9,-9,-9,-9,-9,-9,     // Decimal 179 - 191
+        -9,-9,-9,-9,-9,-9,-9,-9,-9,-9,-9,-9,-9,     // Decimal 192 - 204
+        -9,-9,-9,-9,-9,-9,-9,-9,-9,-9,-9,-9,-9,     // Decimal 205 - 217
+        -9,-9,-9,-9,-9,-9,-9,-9,-9,-9,-9,-9,-9,     // Decimal 218 - 230
+        -9,-9,-9,-9,-9,-9,-9,-9,-9,-9,-9,-9,-9,     // Decimal 231 - 243
+        -9,-9,-9,-9,-9,-9,-9,-9,-9,-9,-9,-9         // Decimal 244 - 255 */
+      };
+
+  /** The web safe decodabet */
+  private final static byte[] WEBSAFE_DECODABET =
+      {-9, -9, -9, -9, -9, -9, -9, -9, -9, // Decimal  0 -  8
+          -5, -5, // Whitespace: Tab and Linefeed
+          -9, -9, // Decimal 11 - 12
+          -5, // Whitespace: Carriage Return
+          -9, -9, -9, -9, -9, -9, -9, -9, -9, -9, -9, -9, -9, // Decimal 14 - 26
+          -9, -9, -9, -9, -9, // Decimal 27 - 31
+          -5, // Whitespace: Space
+          -9, -9, -9, -9, -9, -9, -9, -9, -9, -9, -9, -9, // Decimal 33 - 44
+          62, // Dash '-' sign at decimal 45
+          -9, -9, // Decimal 46-47
+          52, 53, 54, 55, 56, 57, 58, 59, 60, 61, // Numbers zero through nine
+          -9, -9, -9, // Decimal 58 - 60
+          -1, // Equals sign at decimal 61
+          -9, -9, -9, // Decimal 62 - 64
+          0, 1, 2, 3, 4, 5, 6, 7, 8, 9, 10, 11, 12, 13, // Letters 'A' through 'N'
+          14, 15, 16, 17, 18, 19, 20, 21, 22, 23, 24, 25, // Letters 'O' through 'Z'
+          -9, -9, -9, -9, // Decimal 91-94
+          63, // Underscore '_' at decimal 95
+          -9, // Decimal 96
+          26, 27, 28, 29, 30, 31, 32, 33, 34, 35, 36, 37, 38, // Letters 'a' through 'm'
+          39, 40, 41, 42, 43, 44, 45, 46, 47, 48, 49, 50, 51, // Letters 'n' through 'z'
+          -9, -9, -9, -9, -9 // Decimal 123 - 127
+      /*  ,-9,-9,-9,-9,-9,-9,-9,-9,-9,-9,-9,-9,     // Decimal 128 - 139
+        -9,-9,-9,-9,-9,-9,-9,-9,-9,-9,-9,-9,-9,     // Decimal 140 - 152
+        -9,-9,-9,-9,-9,-9,-9,-9,-9,-9,-9,-9,-9,     // Decimal 153 - 165
+        -9,-9,-9,-9,-9,-9,-9,-9,-9,-9,-9,-9,-9,     // Decimal 166 - 178
+        -9,-9,-9,-9,-9,-9,-9,-9,-9,-9,-9,-9,-9,     // Decimal 179 - 191
+        -9,-9,-9,-9,-9,-9,-9,-9,-9,-9,-9,-9,-9,     // Decimal 192 - 204
+        -9,-9,-9,-9,-9,-9,-9,-9,-9,-9,-9,-9,-9,     // Decimal 205 - 217
+        -9,-9,-9,-9,-9,-9,-9,-9,-9,-9,-9,-9,-9,     // Decimal 218 - 230
+        -9,-9,-9,-9,-9,-9,-9,-9,-9,-9,-9,-9,-9,     // Decimal 231 - 243
+        -9,-9,-9,-9,-9,-9,-9,-9,-9,-9,-9,-9         // Decimal 244 - 255 */
+      };
+
+  // Indicates white space in encoding
+  private final static byte WHITE_SPACE_ENC = -5;
+  // Indicates equals sign in encoding
+  private final static byte EQUALS_SIGN_ENC = -1;
+
+  /** Defeats instantiation. */
+  private Base64() {
+  }
+
+  /* ********  E N C O D I N G   M E T H O D S  ******** */
+
+  /**
+   * Encodes up to three bytes of the array <var>source</var>
+   * and writes the resulting four Base64 bytes to <var>destination</var>.
+   * The source and destination arrays can be manipulated
+   * anywhere along their length by specifying
+   * <var>srcOffset</var> and <var>destOffset</var>.
+   * This method does not check to make sure your arrays
+   * are large enough to accommodate <var>srcOffset</var> + 3 for
+   * the <var>source</var> array or <var>destOffset</var> + 4 for
+   * the <var>destination</var> array.
+   * The actual number of significant bytes in your array is
+   * given by <var>numSigBytes</var>.
+   *
+   * @param source the array to convert
+   * @param srcOffset the index where conversion begins
+   * @param numSigBytes the number of significant bytes in your array
+   * @param destination the array to hold the conversion
+   * @param destOffset the index where output will be put
+   * @param alphabet is the encoding alphabet
+   * @return the <var>destination</var> array
+   * @since 1.3
+   */
+  private static byte[] encode3to4(byte[] source, int srcOffset,
+      int numSigBytes, byte[] destination, int destOffset, byte[] alphabet) {
+    //           1         2         3
+    // 01234567890123456789012345678901 Bit position
+    // --------000000001111111122222222 Array position from threeBytes
+    // --------|    ||    ||    ||    | Six bit groups to index alphabet
+    //          >>18  >>12  >> 6  >> 0  Right shift necessary
+    //                0x3f  0x3f  0x3f  Additional AND
+
+    // Create buffer with zero-padding if there are only one or two
+    // significant bytes passed in the array.
+    // We have to shift left 24 in order to flush out the 1's that appear
+    // when Java treats a value as negative that is cast from a byte to an int.
+    int inBuff =
+        (numSigBytes > 0 ? ((source[srcOffset] << 24) >>> 8) : 0)
+            | (numSigBytes > 1 ? ((source[srcOffset + 1] << 24) >>> 16) : 0)
+            | (numSigBytes > 2 ? ((source[srcOffset + 2] << 24) >>> 24) : 0);
+
+    switch (numSigBytes) {
+      case 3:
+        destination[destOffset] = alphabet[(inBuff >>> 18)];
+        destination[destOffset + 1] = alphabet[(inBuff >>> 12) & 0x3f];
+        destination[destOffset + 2] = alphabet[(inBuff >>> 6) & 0x3f];
+        destination[destOffset + 3] = alphabet[(inBuff) & 0x3f];
+        return destination;
+      case 2:
+        destination[destOffset] = alphabet[(inBuff >>> 18)];
+        destination[destOffset + 1] = alphabet[(inBuff >>> 12) & 0x3f];
+        destination[destOffset + 2] = alphabet[(inBuff >>> 6) & 0x3f];
+        destination[destOffset + 3] = EQUALS_SIGN;
+        return destination;
+      case 1:
+        destination[destOffset] = alphabet[(inBuff >>> 18)];
+        destination[destOffset + 1] = alphabet[(inBuff >>> 12) & 0x3f];
+        destination[destOffset + 2] = EQUALS_SIGN;
+        destination[destOffset + 3] = EQUALS_SIGN;
+        return destination;
+      default:
+        return destination;
+    } // end switch
+  } // end encode3to4
+
+  /**
+   * Encodes a byte array into Base64 notation.
+   * Equivalent to calling
+   * {@code encodeBytes(source, 0, source.length)}
+   *
+   * @param source The data to convert
+   * @since 1.4
+   */
+  public static String encode(byte[] source) {
+    return encode(source, 0, source.length, ALPHABET, true);
+  }
+
+  /**
+   * Encodes a byte array into web safe Base64 notation.
+   *
+   * @param source The data to convert
+   * @param doPadding is {@code true} to pad result with '=' chars
+   *        if it does not fall on 3 byte boundaries
+   */
+  public static String encodeWebSafe(byte[] source, boolean doPadding) {
+    return encode(source, 0, source.length, WEBSAFE_ALPHABET, doPadding);
+  }
+
+  /**
+   * Encodes a byte array into Base64 notation.
+   *
+   * @param source The data to convert
+   * @param off Offset in array where conversion should begin
+   * @param len Length of data to convert
+   * @param alphabet is the encoding alphabet
+   * @param doPadding is {@code true} to pad result with '=' chars
+   *        if it does not fall on 3 byte boundaries
+   * @since 1.4
+   */
+  public static String encode(byte[] source, int off, int len, byte[] alphabet,
+      boolean doPadding) {
+    byte[] outBuff = encode(source, off, len, alphabet, Integer.MAX_VALUE);
+    int outLen = outBuff.length;
+
+    // If doPadding is false, set length to truncate '='
+    // padding characters
+    while (doPadding == false && outLen > 0) {
+      if (outBuff[outLen - 1] != '=') {
+        break;
+      }
+      outLen -= 1;
+    }
+
+    return new String(outBuff, 0, outLen);
+  }
+
+  /**
+   * Encodes a byte array into Base64 notation.
+   *
+   * @param source The data to convert
+   * @param off Offset in array where conversion should begin
+   * @param len Length of data to convert
+   * @param alphabet is the encoding alphabet
+   * @param maxLineLength maximum length of one line.
+   * @return the BASE64-encoded byte array
+   */
+  public static byte[] encode(byte[] source, int off, int len, byte[] alphabet,
+      int maxLineLength) {
+    int lenDiv3 = (len + 2) / 3; // ceil(len / 3)
+    int len43 = lenDiv3 * 4;
+    byte[] outBuff = new byte[len43 // Main 4:3
+        + (len43 / maxLineLength)]; // New lines
+
+    int d = 0;
+    int e = 0;
+    int len2 = len - 2;
+    int lineLength = 0;
+    for (; d < len2; d += 3, e += 4) {
+
+      // The following block of code is the same as
+      // encode3to4( source, d + off, 3, outBuff, e, alphabet );
+      // but inlined for faster encoding (~20% improvement)
+      int inBuff =
+          ((source[d + off] << 24) >>> 8)
+              | ((source[d + 1 + off] << 24) >>> 16)
+              | ((source[d + 2 + off] << 24) >>> 24);
+      outBuff[e] = alphabet[(inBuff >>> 18)];
+      outBuff[e + 1] = alphabet[(inBuff >>> 12) & 0x3f];
+      outBuff[e + 2] = alphabet[(inBuff >>> 6) & 0x3f];
+      outBuff[e + 3] = alphabet[(inBuff) & 0x3f];
+
+      lineLength += 4;
+      if (lineLength == maxLineLength) {
+        outBuff[e + 4] = NEW_LINE;
+        e++;
+        lineLength = 0;
+      } // end if: end of line
+    } // end for: each piece of array
+
+    if (d < len) {
+      encode3to4(source, d + off, len - d, outBuff, e, alphabet);
+
+      lineLength += 4;
+      if (lineLength == maxLineLength) {
+        // Add a last newline
+        outBuff[e + 4] = NEW_LINE;
+        e++;
+      }
+      e += 4;
+    }
+
+    assert (e == outBuff.length);
+    return outBuff;
+  }
+
+
+  /* ********  D E C O D I N G   M E T H O D S  ******** */
+
+
+  /**
+   * Decodes four bytes from array <var>source</var>
+   * and writes the resulting bytes (up to three of them)
+   * to <var>destination</var>.
+   * The source and destination arrays can be manipulated
+   * anywhere along their length by specifying
+   * <var>srcOffset</var> and <var>destOffset</var>.
+   * This method does not check to make sure your arrays
+   * are large enough to accommodate <var>srcOffset</var> + 4 for
+   * the <var>source</var> array or <var>destOffset</var> + 3 for
+   * the <var>destination</var> array.
+   * This method returns the actual number of bytes that
+   * were converted from the Base64 encoding.
+   *
+   *
+   * @param source the array to convert
+   * @param srcOffset the index where conversion begins
+   * @param destination the array to hold the conversion
+   * @param destOffset the index where output will be put
+   * @param decodabet the decodabet for decoding Base64 content
+   * @return the number of decoded bytes converted
+   * @since 1.3
+   */
+  private static int decode4to3(byte[] source, int srcOffset,
+      byte[] destination, int destOffset, byte[] decodabet) {
+    // Example: Dk==
+    if (source[srcOffset + 2] == EQUALS_SIGN) {
+      int outBuff =
+          ((decodabet[source[srcOffset]] << 24) >>> 6)
+              | ((decodabet[source[srcOffset + 1]] << 24) >>> 12);
+
+      destination[destOffset] = (byte) (outBuff >>> 16);
+      return 1;
+    } else if (source[srcOffset + 3] == EQUALS_SIGN) {
+      // Example: DkL=
+      int outBuff =
+          ((decodabet[source[srcOffset]] << 24) >>> 6)
+              | ((decodabet[source[srcOffset + 1]] << 24) >>> 12)
+              | ((decodabet[source[srcOffset + 2]] << 24) >>> 18);
+
+      destination[destOffset] = (byte) (outBuff >>> 16);
+      destination[destOffset + 1] = (byte) (outBuff >>> 8);
+      return 2;
+    } else {
+      // Example: DkLE
+      int outBuff =
+          ((decodabet[source[srcOffset]] << 24) >>> 6)
+              | ((decodabet[source[srcOffset + 1]] << 24) >>> 12)
+              | ((decodabet[source[srcOffset + 2]] << 24) >>> 18)
+              | ((decodabet[source[srcOffset + 3]] << 24) >>> 24);
+
+      destination[destOffset] = (byte) (outBuff >> 16);
+      destination[destOffset + 1] = (byte) (outBuff >> 8);
+      destination[destOffset + 2] = (byte) (outBuff);
+      return 3;
+    }
+  } // end decodeToBytes
+
+
+  /**
+   * Decodes data from Base64 notation.
+   *
+   * @param s the string to decode (decoded in default encoding)
+   * @return the decoded data
+   * @since 1.4
+   */
+  public static byte[] decode(String s) throws Base64DecoderException {
+    byte[] bytes = s.getBytes();
+    return decode(bytes, 0, bytes.length);
+  }
+
+  /**
+   * Decodes data from web safe Base64 notation.
+   * Web safe encoding uses '-' instead of '+', '_' instead of '/'
+   *
+   * @param s the string to decode (decoded in default encoding)
+   * @return the decoded data
+   */
+  public static byte[] decodeWebSafe(String s) throws Base64DecoderException {
+    byte[] bytes = s.getBytes();
+    return decodeWebSafe(bytes, 0, bytes.length);
+  }
+
+  /**
+   * Decodes Base64 content in byte array format and returns
+   * the decoded byte array.
+   *
+   * @param source The Base64 encoded data
+   * @return decoded data
+   * @since 1.3
+   * @throws Base64DecoderException
+   */
+  public static byte[] decode(byte[] source) throws Base64DecoderException {
+    return decode(source, 0, source.length);
+  }
+
+  /**
+   * Decodes web safe Base64 content in byte array format and returns
+   * the decoded data.
+   * Web safe encoding uses '-' instead of '+', '_' instead of '/'
+   *
+   * @param source the string to decode (decoded in default encoding)
+   * @return the decoded data
+   */
+  public static byte[] decodeWebSafe(byte[] source)
+      throws Base64DecoderException {
+    return decodeWebSafe(source, 0, source.length);
+  }
+
+  /**
+   * Decodes Base64 content in byte array format and returns
+   * the decoded byte array.
+   *
+   * @param source The Base64 encoded data
+   * @param off    The offset of where to begin decoding
+   * @param len    The length of characters to decode
+   * @return decoded data
+   * @since 1.3
+   * @throws Base64DecoderException
+   */
+  public static byte[] decode(byte[] source, int off, int len)
+      throws Base64DecoderException {
+    return decode(source, off, len, DECODABET);
+  }
+
+  /**
+   * Decodes web safe Base64 content in byte array format and returns
+   * the decoded byte array.
+   * Web safe encoding uses '-' instead of '+', '_' instead of '/'
+   *
+   * @param source The Base64 encoded data
+   * @param off    The offset of where to begin decoding
+   * @param len    The length of characters to decode
+   * @return decoded data
+   */
+  public static byte[] decodeWebSafe(byte[] source, int off, int len)
+      throws Base64DecoderException {
+    return decode(source, off, len, WEBSAFE_DECODABET);
+  }
+
+  /**
+   * Decodes Base64 content using the supplied decodabet and returns
+   * the decoded byte array.
+   *
+   * @param source    The Base64 encoded data
+   * @param off       The offset of where to begin decoding
+   * @param len       The length of characters to decode
+   * @param decodabet the decodabet for decoding Base64 content
+   * @return decoded data
+   */
+  public static byte[] decode(byte[] source, int off, int len, byte[] decodabet)
+      throws Base64DecoderException {
+    int len34 = len * 3 / 4;
+    byte[] outBuff = new byte[2 + len34]; // Upper limit on size of output
+    int outBuffPosn = 0;
+
+    byte[] b4 = new byte[4];
+    int b4Posn = 0;
+    int i = 0;
+    byte sbiCrop = 0;
+    byte sbiDecode = 0;
+    for (i = 0; i < len; i++) {
+      sbiCrop = (byte) (source[i + off] & 0x7f); // Only the low seven bits
+      sbiDecode = decodabet[sbiCrop];
+
+      if (sbiDecode >= WHITE_SPACE_ENC) { // White space Equals sign or better
+        if (sbiDecode >= EQUALS_SIGN_ENC) {
+          // An equals sign (for padding) must not occur at position 0 or 1
+          // and must be the last byte[s] in the encoded value
+          if (sbiCrop == EQUALS_SIGN) {
+            int bytesLeft = len - i;
+            byte lastByte = (byte) (source[len - 1 + off] & 0x7f);
+            if (b4Posn == 0 || b4Posn == 1) {
+              throw new Base64DecoderException(
+                  "invalid padding byte '=' at byte offset " + i);
+            } else if ((b4Posn == 3 && bytesLeft > 2)
+                || (b4Posn == 4 && bytesLeft > 1)) {
+              throw new Base64DecoderException(
+                  "padding byte '=' falsely signals end of encoded value "
+                      + "at offset " + i);
+            } else if (lastByte != EQUALS_SIGN && lastByte != NEW_LINE) {
+              throw new Base64DecoderException(
+                  "encoded value has invalid trailing byte");
+            }
+            break;
+          }
+
+          b4[b4Posn++] = sbiCrop;
+          if (b4Posn == 4) {
+            outBuffPosn += decode4to3(b4, 0, outBuff, outBuffPosn, decodabet);
+            b4Posn = 0;
+          }
+        }
+      } else {
+        throw new Base64DecoderException("Bad Base64 input character at " + i
+            + ": " + source[i + off] + "(decimal)");
+      }
+    }
+
+    // Because web safe encoding allows non padding base64 encodes, we
+    // need to pad the rest of the b4 buffer with equal signs when
+    // b4Posn != 0.  There can be at most 2 equal signs at the end of
+    // four characters, so the b4 buffer must have two or three
+    // characters.  This also catches the case where the input is
+    // padded with EQUALS_SIGN
+    if (b4Posn != 0) {
+      if (b4Posn == 1) {
+        throw new Base64DecoderException("single trailing character at offset "
+            + (len - 1));
+      }
+      b4[b4Posn++] = EQUALS_SIGN;
+      outBuffPosn += decode4to3(b4, 0, outBuff, outBuffPosn, decodabet);
+    }
+
+    byte[] out = new byte[outBuffPosn];
+    System.arraycopy(outBuff, 0, out, 0, outBuffPosn);
+    return out;
+  }
+}
diff --git a/Engine/Build/Android/Java/src/com/android/vending/licensing/util/Base64DecoderException.java b/Engine/Build/Android/Java/src/com/android/vending/licensing/util/Base64DecoderException.java
new file mode 100644
--- /dev/null
+++ b/Engine/Build/Android/Java/src/com/android/vending/licensing/util/Base64DecoderException.java
@@ -0,0 +1,32 @@
+// Copyright 2002, Google, Inc.
+//
+// Licensed under the Apache License, Version 2.0 (the "License");
+// you may not use this file except in compliance with the License.
+// You may obtain a copy of the License at
+//
+//     http://www.apache.org/licenses/LICENSE-2.0
+//
+// Unless required by applicable law or agreed to in writing, software
+// distributed under the License is distributed on an "AS IS" BASIS,
+// WITHOUT WARRANTIES OR CONDITIONS OF ANY KIND, either express or implied.
+// See the License for the specific language governing permissions and
+// limitations under the License.
+
+package com.google.android.vending.licensing.util;
+
+/**
+ * Exception thrown when encountering an invalid Base64 input character.
+ *
+ * @author nelson
+ */
+public class Base64DecoderException extends Exception {
+  public Base64DecoderException() {
+    super();
+  }
+
+  public Base64DecoderException(String s) {
+    super(s);
+  }
+
+  private static final long serialVersionUID = 1L;
+}
