In-App Purchase changes.

Android fixes for workflow and threading issues.
Receipt is now returned as part of the IAP product information in blueprints
IAP workflow in blueprints now use an IAP request object as opposed to just a string.
Changes to the IOS workflow to accommodate receipts in product responses.
Android game activity has been updated to receive onActivityChanged messages. Needed for routing activity changes to android IAP.
A lot of threading fixes
Fixes for certain google events not calling the appropriate delegates.
Play store license key property added to runtime settings.

diff --git a/Engine/Build/Android/Java/src/com/epicgames/ue4/GameActivity.java b/Engine/Build/Android/Java/src/com/epicgames/ue4/GameActivity.java
--- a/Engine/Build/Android/Java/src/com/epicgames/ue4/GameActivity.java
+++ b/Engine/Build/Android/Java/src/com/epicgames/ue4/GameActivity.java
@@ -830,6 +830,19 @@ public class GameActivity extends NativeActivity
 		}
 		return bTriggeredQuery;
 	}
+
+	@Override
+	protected void onActivityResult(int requestCode, int resultCode, Intent data)
+	{
+		if(!IapStoreHelper.onActivityResult(requestCode, resultCode, data))
+		{
+			super.onActivityResult(requestCode, resultCode, data);
+		}
+		else
+		{
+			Log.debug("[JAVA] - Store Helper handled onActivityResult");
+		}
+	}
 	
 	public boolean AndroidThunkJava_IapBeginPurchase(String ProductId, boolean bConsumable)
 	{
diff --git a/Engine/Build/Android/Java/src/com/epicgames/ue4/GooglePlayStoreHelper.java b/Engine/Build/Android/Java/src/com/epicgames/ue4/GooglePlayStoreHelper.java
--- a/Engine/Build/Android/Java/src/com/epicgames/ue4/GooglePlayStoreHelper.java
+++ b/Engine/Build/Android/Java/src/com/epicgames/ue4/GooglePlayStoreHelper.java
@@ -3,7 +3,13 @@
 
 package com.epicgames.ue4;
 
+import java.util.ArrayList;
+
 import android.os.RemoteException;
+import org.json.JSONException;
+import android.os.Handler;
+import android.os.Looper;
+import android.content.Intent;
 import android.content.IntentSender.SendIntentException;
 
 import com.android.vending.billing.util.IabHelper;
@@ -21,40 +27,33 @@ public class GooglePlayStoreHelper extends StoreHelper
 	private Logger Log;
 
 	private String ProductKey;
+
 	
-	// Iab helper functions for the store interface
-	IabHelper.OnConsumeFinishedListener ConsumeProductFinishedCallback =  new IabHelper.OnConsumeFinishedListener() 
-	{
-		public void onConsumeFinished(Purchase purchase, IabResult result) 
-		{
-			if (result.isSuccess()) 
-			{
-				Log.debug("[GooglePlayStoreHelper] - onConsumeFinished: " + (result.isSuccess() ? "Success" : "Failed" ));
-				String Receipt = Base64.encode(purchase.getOriginalJson().getBytes());
-				nativePurchaseComplete(true, purchase.getSku(), Receipt);
-			}
-			else
-			{
-				// Something went wrong, cannot consume product...send back success, but will be unable to purchase again
-				Log.debug("[GooglePlayStoreHelper] - ERROR: Could not successfully consume item. Will attempt again on next inventory update.");
-				nativePurchaseComplete(false, purchase.getSku(), "");
-			}
-		}
-	};
+	
+	/*********************************************************************
+	 ** IAP SETUP WORKFLOW...
+	 *********************************************************************/
+	private boolean bIsIapSetup = false;
 
+	
+	// Constructor...
 	public GooglePlayStoreHelper(String InProductKey, GameActivity InGameActivity, final Logger InLog)
 	{
 		ProductKey = InProductKey;
 		gameActivity = InGameActivity;
 		Log = InLog;
+
+		// Create an Iab Helper.
 		SetupIapService();
 	}
+	
 
-	private boolean bIsIapSetup = false;
+	/**
+	 *	Initialize the Iab helper functionality for our application.
+	 */
 	public void SetupIapService()
 	{
 		Log.debug("[GooglePlayStoreHelper] - GooglePlayStoreHelper::SetupIapService");
-
 		inAppPurchaseHelper = new IabHelper(gameActivity, ProductKey);
 
 		inAppPurchaseHelper.startSetup(new IabHelper.OnIabSetupFinishedListener() 
@@ -77,9 +76,80 @@ public class GooglePlayStoreHelper extends StoreHelper
 		});
 	}
 
-
+	
+	/*********************************************************************
+	 ** QUERY PURCHASES WORKFLOW...
+	 *********************************************************************/
 	private String[] cachedQueryProductIds;
 	private boolean[] cachedQueryConsumableFlags;
+	
+	
+	// Iab helper functions for the store interface and query functionality
+	IabHelper.QueryInventoryFinishedListener GotInventoryListener = new IabHelper.QueryInventoryFinishedListener()
+	{
+		public void onQueryInventoryFinished(IabResult result, Inventory inventory)
+		{
+			Log.debug("[GooglePlayStoreHelper] - GooglePlayStoreHelper::onQueryInventoryFinished - " + result);
+
+			if (result.isSuccess())
+			{			
+				ArrayList<String> titles = new ArrayList<String>();
+				ArrayList<String> descriptions = new ArrayList<String>();
+				ArrayList<String> prices = new ArrayList<String>();
+				for( int Idx = 0; Idx < cachedQueryProductIds.length; Idx++ )
+				{	
+					String productId = cachedQueryProductIds[Idx];
+					Log.debug("[GooglePlayStoreHelper] - Checking details for " + productId);
+
+					if( inventory.hasDetails(productId) )
+					{
+					Log.debug("[GooglePlayStoreHelper] - Found details for " + productId);
+						titles.add( inventory.getSkuDetails(productId).getTitle() );
+						descriptions.add( inventory.getSkuDetails(productId).getDescription() );
+						prices.add( inventory.getSkuDetails(productId).getPrice() );
+					}
+					
+					// If we have this item, but somehow failed to consume it on purchase.
+					// make sure to consume it now. TPMB - TODO: make this part of the purchase workflow?
+					if (inventory.hasPurchase(productId))
+					{
+						try {
+							Purchase purchase = inventory.getPurchase(productId);
+							inAppPurchaseHelper.consumeAsync(purchase, ConsumeProductFinishedCallback);
+						} catch(Exception e) {
+							Log.debug("[GooglePlayStoreHelper] - consumeAsync - Exception caught - " + e.getMessage());
+						}
+					}
+				}
+
+				if( titles.size() > 0 )
+				{
+					if( titles.size() != cachedQueryProductIds.length )
+					{
+						Log.debug("[GooglePlayStoreHelper] - Mismatch in the number of products returned vs those queried.");
+						Log.debug("[GooglePlayStoreHelper] - Requested: "+ cachedQueryProductIds.length + ", Received information for: " + titles.size());
+					}
+					Log.debug("[GooglePlayStoreHelper] - GooglePlayStoreHelper::onQueryInventoryFinished - Succeeded.");
+					nativeQueryComplete(true, cachedQueryProductIds, titles.toArray(new String[titles.size()]), descriptions.toArray(new String[descriptions.size()]), prices.toArray(new String[prices.size()])); 
+				}
+				else
+				{
+					Log.debug("[GooglePlayStoreHelper] - GooglePlayStoreHelper::onQueryInventoryFinished - Failed. " + result.getMessage());
+					nativeQueryComplete(false, cachedQueryProductIds, null, null, null);
+				}
+			}
+			else
+			{
+				Log.debug("[GooglePlayStoreHelper] - GooglePlayStoreHelper::onQueryInventoryFinished - Failed. " + result.getMessage());
+				nativeQueryComplete(false, cachedQueryProductIds, null, null, null);
+			}
+		}
+	};
+
+
+	/**
+	 *	Execute a query available purchases workflow for the provided product Id
+	 */
 	public boolean QueryInAppPurchases(final String[] ProductIds, final boolean[] bConsumableFlags)
 	{
 		Log.debug("[GooglePlayStoreHelper] - GooglePlayStoreHelper::QueryInAppPurchases");
@@ -92,44 +162,12 @@ public class GooglePlayStoreHelper extends StoreHelper
 		boolean bIsAbleToSendQuery = bIsIapSetup && (cachedQueryProductIds.length > 0 );
 		if( bIsAbleToSendQuery )
 		{
-			IabHelper.QueryInventoryFinishedListener GotInventoryListener = new IabHelper.QueryInventoryFinishedListener()
-			{
-				public void onQueryInventoryFinished(IabResult result, Inventory inventory)
-				{
-					if (result.isSuccess())
-					{
-						Log.debug("[GooglePlayStoreHelper] - GooglePlayStoreHelper::onQueryInventoryFinished - Success.");
-						String[] titles = new String[cachedQueryProductIds.length];
-						String[] descriptions = new String[cachedQueryProductIds.length];
-						String[] prices = new String[cachedQueryProductIds.length];
-
-						int Idx = 0;
-						for( String productId : cachedQueryProductIds )
-						{
-							titles[Idx] = inventory.getSkuDetails(productId).getTitle();
-							descriptions[Idx] = inventory.getSkuDetails(productId).getDescription();
-							prices[Idx] = inventory.getSkuDetails(productId).getPrice();
-
-							++Idx;
-						}
-
-						Log.debug("[GooglePlayStoreHelper] - GooglePlayStoreHelper::onQueryInventoryFinished - Succeeded.");
-						nativeQueryComplete(true, cachedQueryProductIds, titles, descriptions, prices); 
-					}
-					else
-					{
-						Log.debug("[GooglePlayStoreHelper] - GooglePlayStoreHelper::onQueryInventoryFinished - Failed.");
-						nativeQueryComplete(false, null, null, null, null);
-					}
-				}
-			};
-
 			// Run the query pruchases command.
 			Log.debug("[GooglePlayStoreHelper] - queryInventoryAsync. - start");
 			try {
 				inAppPurchaseHelper.queryInventoryAsync(GotInventoryListener);
 			} catch(Exception e) {
-				Log.debug("[GooglePlayStoreHelper] - Exception caught - " + e.getMessage());
+				Log.debug("[GooglePlayStoreHelper] - queryInventoryAsync - Exception caught - " + e.getMessage());
 			}
 			Log.debug("[GooglePlayStoreHelper] - queryInventoryAsync. - end");
 
@@ -148,59 +186,152 @@ public class GooglePlayStoreHelper extends StoreHelper
 			nativeQueryComplete(false, null, null, null, null);
 		}
 
-
 		return bIsAbleToSendQuery;
 	}
 
+	
+	/*********************************************************************
+	 ** PURCHASE WORKFLOW...
+	 *********************************************************************/
 	String cachedPurchaseProductId;
 	boolean cachedPurchaseConsumableFlag;
-	public boolean BeginPurchase(final String ProductId, final boolean bConsumableFlag)
+
+	// Iab helper functions for the store interface
+	IabHelper.OnConsumeFinishedListener ConsumeProductFinishedCallback =  new IabHelper.OnConsumeFinishedListener() 
 	{
-		Log.debug("[GooglePlayStoreHelper] - GooglePlayStoreHelper::BeginPurchase");
+		public void onConsumeFinished(Purchase purchase, IabResult result) 
+		{
+			Log.debug("[GooglePlayStoreHelper] - ConsumeProductFinishedCallback::onConsumeFinished");
+			if (result.isSuccess()) 
+			{
+				Log.debug("[GooglePlayStoreHelper] - onConsumeFinished: Success");
+				String Receipt = Base64.encode(purchase.getOriginalJson().getBytes());
+				nativePurchaseComplete(true, purchase.getSku(), Receipt);
+			}
+			else
+			{
+				// Something went wrong, cannot consume product...send back success, but will be unable to purchase again
+				Log.debug("[GooglePlayStoreHelper] - ERROR: Could not successfully consume item. Will attempt again on next inventory update.");
+				nativePurchaseComplete(false, purchase.getSku(), "");
+			}
+		}
+	};
 	
-		// Cache the input parameters for cross threading
-		cachedPurchaseProductId = ProductId;
-		cachedPurchaseConsumableFlag = bConsumableFlag;
 
-		// We can launch a purchase workflow if we have a provided product Id and the IabHelper is set-up.
-		boolean bIsValidProductId = (cachedPurchaseProductId.length() > 0);
-		if( bIsIapSetup && bIsValidProductId )
+	// Query inventory callback used to find an existing purchase and consume it.
+	// Only used when we have a lingering purchase on the server that wasnt consumed.
+	IabHelper.QueryInventoryFinishedListener ConsumeInventoryListener = new IabHelper.QueryInventoryFinishedListener()
+	{
+		public void onQueryInventoryFinished(IabResult result, Inventory inventory)
+		{
+			Log.debug("[GooglePlayStoreHelper] - ConsumeInventoryListener::onQueryInventoryFinished");
+			if (result.isSuccess() && inventory.hasPurchase(cachedPurchaseProductId))
+			{
+				final Purchase UnconsumedPurchase = inventory.getPurchase(cachedPurchaseProductId);
+				inAppPurchaseHelper.consumeAsync(UnconsumedPurchase, ConsumeProductFinishedCallback);										
+			}
+			else
+			{
+				Log.debug("[GooglePlayStoreHelper] - Purchase Unsuccessful. Failed to find the already purchased product on the server...");
+				nativePurchaseComplete(false, cachedPurchaseProductId, "" );
+			}
+		}
+	};
+
+	
+	// Iab helper functions for the store interface and purchase functionality
+	IabHelper.OnIabPurchaseFinishedListener PurchaseFinishedCallback = new IabHelper.OnIabPurchaseFinishedListener() 
+	{
+		public void onIabPurchaseFinished(IabResult result, Purchase purchase)
 		{
-			IabHelper.OnIabPurchaseFinishedListener PurchaseFinishedCallback = new IabHelper.OnIabPurchaseFinishedListener() 
+			Log.debug("[GooglePlayStoreHelper] - GooglePlayStoreHelper::onIabPurchaseFinished - " + result + ", purchase: " + purchase);
+			if(result.isSuccess())
 			{
-				public void onIabPurchaseFinished(IabResult result, Purchase purchase) 
+				if( VerifyPayload(purchase.getDeveloperPayload(), purchase.getSku() ) == true )
 				{
-					if(result.isSuccess())
+					// Consumable items must be consumed on the server prior to provisioning the item to the game
+					if(cachedPurchaseConsumableFlag == true)
 					{
-						// Consumable items must be consumed on the server prior to provisioning the item to the game
-						if(cachedPurchaseConsumableFlag == true)
-						{
-							Log.debug("[GooglePlayStoreHelper] - Purchase Success. Requesting Item be consumed.");
-							inAppPurchaseHelper.consumeAsync(purchase, ConsumeProductFinishedCallback);
-						}
-						// Non consumable can immediately notify the game of the purchase
-						else
-						{
-							Log.debug("[GooglePlayStoreHelper] - Purchase Success.");
-							String Receipt = Base64.encode(purchase.getOriginalJson().getBytes());
-							nativePurchaseComplete(false, cachedPurchaseProductId, Receipt);
-						}
+						Log.debug("[GooglePlayStoreHelper] - Purchase Success. Requesting Item be consumed.");
+						inAppPurchaseHelper.consumeAsync(purchase, ConsumeProductFinishedCallback);
 					}
+					// Non consumable can immediately notify the game of the purchase
 					else
 					{
-						Log.debug("[GooglePlayStoreHelper] - Purchase Unsuccessful. " + result.getMessage());
-						nativePurchaseComplete(false, cachedPurchaseProductId, "" );
-					} 
+						Log.debug("[GooglePlayStoreHelper] - Purchase Success.");
+						String Receipt = Base64.encode(purchase.getOriginalJson().getBytes());
+						nativePurchaseComplete(false, cachedPurchaseProductId, Receipt);
+					}
 				}
-			};
-			
-			Log.debug("[GooglePlayStoreHelper] - launchPurchaseFlow. - start");
+				else
+				{
+					Log.debug("[GooglePlayStoreHelper] - Purchase Unsuccessful. Validation of developer payload failed.");
+					nativePurchaseComplete(false, cachedPurchaseProductId, "" );
+				}
+			}
+			else if( result.getResponse() == IabHelper.BILLING_RESPONSE_RESULT_ITEM_ALREADY_OWNED)
+			{
+				Log.debug("[GooglePlayStoreHelper] - Item was already owned but not consumed, attempting to consume... ");
+				
+				// Post the consume command so it is spawned from the main looper, needed for an async operation
+				Handler mainHandler = new Handler(Looper.getMainLooper());
+				mainHandler.post(new Runnable()
+				{
+					@Override
+					public void run()
+					{
+						try {
+								inAppPurchaseHelper.queryInventoryAsync(ConsumeInventoryListener);										
+							}
+						catch(Exception e) {
+							Log.debug("[GooglePlayStoreHelper] - queryInventoryAsync - Exception caught - " + e.getMessage());
+						}
+					}
+				});
+			}
+			else
+			{
+				Log.debug("[GooglePlayStoreHelper] - Purchase Unsuccessful. " + result.getMessage());
+				nativePurchaseComplete(false, cachedPurchaseProductId, "" );
+			} 
+		}
+	};
+
+	private String GenerateDevPayload( String ProductId )
+	{
+		return "ue4."+ProductId;
+	}
+
+	private boolean VerifyPayload( String ExistingPayload, String ProductId )
+	{
+		String GeneratedPayload = GenerateDevPayload(ProductId);
+		Log.debug("[GooglePlayStoreHelper] - GooglePlayStoreHelper::VerifyPayload, ExistingPayload: " + ExistingPayload);
+		Log.debug("[GooglePlayStoreHelper] - GooglePlayStoreHelper::VerifyPayload, GeneratedPayload: " + GeneratedPayload);
+
+		return ExistingPayload.equals(GeneratedPayload);
+	}
+
+	/**
+	 *	Execute a purchase workflow for the provided product Id
+	 */
+	public boolean BeginPurchase(final String ProductId, final boolean bConsumableFlag)
+	{
+		Log.debug("[GooglePlayStoreHelper] - GooglePlayStoreHelper::BeginPurchase");
+	
+		// Cache the input parameters for cross threading
+		cachedPurchaseProductId = ProductId;
+		cachedPurchaseConsumableFlag = bConsumableFlag;
+
+		// We can launch a purchase workflow if we have a provided product Id and the IabHelper is set-up.
+		boolean bIsValidProductId = (cachedPurchaseProductId.length() > 0);
+		if( bIsIapSetup && bIsValidProductId )
+		{	
 			try {
-				inAppPurchaseHelper.launchPurchaseFlow(gameActivity, cachedPurchaseProductId, cachedPurchaseConsumableFlag ? 10001 : 10002, PurchaseFinishedCallback, "");
+				String DevPayload = GenerateDevPayload(ProductId);
+				inAppPurchaseHelper.launchPurchaseFlow(gameActivity, cachedPurchaseProductId, cachedPurchaseConsumableFlag ? 10001 : 10002, PurchaseFinishedCallback, DevPayload);
 			} catch(Exception e) {
-				Log.debug("[GooglePlayStoreHelper] - Exception caught - " + e.getMessage());
+				Log.debug("[GooglePlayStoreHelper] - launchPurchaseFlow - Exception caught - " + e.getMessage());
 			}
-			Log.debug("[GooglePlayStoreHelper] - launchPurchaseFlow. - end");
 		}
 		else
 		{
@@ -226,6 +357,12 @@ public class GooglePlayStoreHelper extends StoreHelper
 		return bIsIapSetup;
 	}
 
+	public boolean onActivityResult(int requestCode, int resultCode, Intent data)
+	{
+		Log.debug("[GooglePlayStoreHelper] - GooglePlayStoreHelper::onActivityResult");
+		return inAppPurchaseHelper != null && inAppPurchaseHelper.handleActivityResult(requestCode, resultCode, data);
+	}
+
 	public native void nativeQueryComplete(boolean bSuccess, String[] productIDs, String[] titles, String[] descriptions, String[] prices );
 
 	public native void nativePurchaseComplete(boolean bSuccess, String ProductId, String ReceiptData);
diff --git a/Engine/Build/Android/Java/src/com/epicgames/ue4/StoreHelper.java b/Engine/Build/Android/Java/src/com/epicgames/ue4/StoreHelper.java
--- a/Engine/Build/Android/Java/src/com/epicgames/ue4/StoreHelper.java
+++ b/Engine/Build/Android/Java/src/com/epicgames/ue4/StoreHelper.java
@@ -12,6 +12,7 @@ public abstract class StoreHelper
 	
 	public abstract void SetupIapService();
 	public abstract boolean QueryInAppPurchases(String[] ProductIDs, boolean[] bConsumable);
-	public abstract boolean BeginPurchase(String ProductIDs, boolean bConsumable);
+	public abstract boolean BeginPurchase(String ProductID, boolean bConsumable);
 	public abstract boolean IsAllowedToMakePurchases();
+	public abstract boolean onActivityResult(int requestCode, int resultCode, Intent data);
 }
\ No newline at end of file
diff --git a/Engine/Source/Developer/Android/AndroidPlatformEditor/Private/AndroidTargetSettingsCustomization.cpp b/Engine/Source/Developer/Android/AndroidPlatformEditor/Private/AndroidTargetSettingsCustomization.cpp
--- a/Engine/Source/Developer/Android/AndroidPlatformEditor/Private/AndroidTargetSettingsCustomization.cpp
+++ b/Engine/Source/Developer/Android/AndroidPlatformEditor/Private/AndroidTargetSettingsCustomization.cpp
@@ -162,6 +162,10 @@ void FAndroidTargetSettingsCustomization::BuildAppManifestSection(IDetailLayoutB
 	TSharedRef<IPropertyHandle> AdMobAdUnitIDProperty = DetailLayout.GetProperty(GET_MEMBER_NAME_CHECKED(UAndroidRuntimeSettings, AdMobAdUnitID));
 	GooglePlayCategory.AddProperty(AdMobAdUnitIDProperty)
 		.EditCondition(SetupForGooglePlayAttribute, NULL);
+
+	TSharedRef<IPropertyHandle> GooglePlayLicenseKeyProperty = DetailLayout.GetProperty(GET_MEMBER_NAME_CHECKED(UAndroidRuntimeSettings, GooglePlayLicenseKey));
+	GooglePlayCategory.AddProperty(GooglePlayLicenseKeyProperty)
+		.EditCondition(SetupForGooglePlayAttribute, NULL);
 }
 
 void FAndroidTargetSettingsCustomization::BuildIconSection(IDetailLayoutBuilder& DetailLayout)
diff --git a/Engine/Source/Runtime/Android/AndroidRuntimeSettings/Classes/AndroidRuntimeSettings.h b/Engine/Source/Runtime/Android/AndroidRuntimeSettings/Classes/AndroidRuntimeSettings.h
--- a/Engine/Source/Runtime/Android/AndroidRuntimeSettings/Classes/AndroidRuntimeSettings.h
+++ b/Engine/Source/Runtime/Android/AndroidRuntimeSettings/Classes/AndroidRuntimeSettings.h
@@ -150,6 +150,10 @@ public:
 	UPROPERTY(GlobalConfig, EditAnywhere, Category = GooglePlayServices)
 	FString AdMobAdUnitID;
 
+	// The unique identifier for the ad obtained from AdMob.
+	UPROPERTY(GlobalConfig, EditAnywhere, Category = GooglePlayServices)
+	FString GooglePlayLicenseKey;
+
 	/** Android Audio encoding options */
 	UPROPERTY(GlobalConfig, EditAnywhere, Category = DataCooker, meta = (DisplayName = "Audio encoding"))
 	TEnumAsByte<EAndroidAudio::Type> AndroidAudio;
diff --git a/Engine/Source/Runtime/Online/Android/OnlineSubsystemGooglePlay/Private/OnlineAsyncTaskGooglePlayQueryInAppPurchases.cpp b/Engine/Source/Runtime/Online/Android/OnlineSubsystemGooglePlay/Private/OnlineAsyncTaskGooglePlayQueryInAppPurchases.cpp
--- a/Engine/Source/Runtime/Online/Android/OnlineSubsystemGooglePlay/Private/OnlineAsyncTaskGooglePlayQueryInAppPurchases.cpp
+++ b/Engine/Source/Runtime/Online/Android/OnlineSubsystemGooglePlay/Private/OnlineAsyncTaskGooglePlayQueryInAppPurchases.cpp
@@ -10,7 +10,7 @@ FOnlineAsyncTaskGooglePlayQueryInAppPurchases::FOnlineAsyncTaskGooglePlayQueryIn
 	: FOnlineAsyncTaskBasic(InSubsystem)
 	, ProductIds(InProductIds)
 	, IsConsumableFlags(InIsConsumableFlags)
-
+	, bWasRequestSent(false)
 {
 }
 
@@ -39,10 +39,15 @@ void FOnlineAsyncTaskGooglePlayQueryInAppPurchases::TriggerDelegates()
 
 void FOnlineAsyncTaskGooglePlayQueryInAppPurchases::Tick()
 {
-	UE_LOG(LogOnline, Display, TEXT("FOnlineAsyncTaskGooglePlayQueryInAppPurchases::Tick"));
-	bWasSuccessful = false;
+//	UE_LOG(LogOnline, Display, TEXT("FOnlineAsyncTaskGooglePlayQueryInAppPurchases::Tick"));
+
+	if (!bWasRequestSent)
+	{
+		extern bool AndroidThunkCpp_Iap_QueryInAppPurchases(const TArray<FString>&, const TArray<bool>&);
+		AndroidThunkCpp_Iap_QueryInAppPurchases(ProductIds, IsConsumableFlags);
+
+		bWasRequestSent = true;
+	}
 
-	extern bool AndroidThunkCpp_Iap_QueryInAppPurchases(const TArray<FString>&, const TArray<bool>&);
-	AndroidThunkCpp_Iap_QueryInAppPurchases(ProductIds, IsConsumableFlags);
 }
 
diff --git a/Engine/Source/Runtime/Online/Android/OnlineSubsystemGooglePlay/Private/OnlineAsyncTaskGooglePlayQueryInAppPurchases.h b/Engine/Source/Runtime/Online/Android/OnlineSubsystemGooglePlay/Private/OnlineAsyncTaskGooglePlayQueryInAppPurchases.h
--- a/Engine/Source/Runtime/Online/Android/OnlineSubsystemGooglePlay/Private/OnlineAsyncTaskGooglePlayQueryInAppPurchases.h
+++ b/Engine/Source/Runtime/Online/Android/OnlineSubsystemGooglePlay/Private/OnlineAsyncTaskGooglePlayQueryInAppPurchases.h
@@ -32,4 +32,7 @@ private:
 
 	// The consume flags provided for this task
 	const TArray<bool> IsConsumableFlags;
+
+	/** Flag indicating that the request has been sent */
+	bool bWasRequestSent;
 };
diff --git a/Engine/Source/Runtime/Online/Android/OnlineSubsystemGooglePlay/Private/OnlineStoreInterfaceGooglePlay.cpp b/Engine/Source/Runtime/Online/Android/OnlineSubsystemGooglePlay/Private/OnlineStoreInterfaceGooglePlay.cpp
--- a/Engine/Source/Runtime/Online/Android/OnlineSubsystemGooglePlay/Private/OnlineStoreInterfaceGooglePlay.cpp
+++ b/Engine/Source/Runtime/Online/Android/OnlineSubsystemGooglePlay/Private/OnlineStoreInterfaceGooglePlay.cpp
@@ -2,6 +2,7 @@
 
 #include "OnlineSubsystemGooglePlayPrivatePCH.h"
 #include "OnlineAsyncTaskGooglePlayQueryInAppPurchases.h"
+#include "TaskGraphInterfaces.h"
 #include <jni.h>
 
 
@@ -15,11 +16,11 @@ FOnlineStoreGooglePlay::FOnlineStoreGooglePlay(FOnlineSubsystemGooglePlay* InSub
 {
 	UE_LOG(LogOnline, Display, TEXT( "FOnlineStoreGooglePlay::FOnlineStoreGooglePlay" ));
 	
-	FString ProductKey;
-	GConfig->GetString(TEXT("/Script/AndroidRuntimeSettings.AndroidRuntimeSettings"), TEXT("GooglePlayProductKey"), ProductKey, GEngineIni);
+	FString LicenseKey;
+	GConfig->GetString(TEXT("/Script/AndroidRuntimeSettings.AndroidRuntimeSettings"), TEXT("GooglePlayLicenseKey"), LicenseKey, GEngineIni);
 
 	extern void AndroidThunkCpp_Iap_SetupIapService(const FString&);
-	AndroidThunkCpp_Iap_SetupIapService(ProductKey);
+	AndroidThunkCpp_Iap_SetupIapService(LicenseKey);
 }
 
 
@@ -111,16 +112,26 @@ extern "C" void Java_com_epicgames_ue4_GooglePlayStoreHelper_nativeQueryComplete
 		}
 	}
 
-	if (IOnlineSubsystem* const OnlineSub = IOnlineSubsystem::Get())
-	{
-		FOnlineStoreGooglePlay* StoreInterface = (FOnlineStoreGooglePlay*)OnlineSub->GetStoreInterface().Get();
-		if (StoreInterface)
-		{
-			StoreInterface->ProcessQueryAvailablePurchasesResults(bSuccess, ProvidedProductInformation);
-		}
-	}
 
-	FPlatformMisc::LowLevelOutputDebugStringf(TEXT("In-App Purchase query was completed  %s\n"), bSuccess ? TEXT("successfully") : TEXT("unsuccessfully"));
+	DECLARE_CYCLE_STAT(TEXT("FSimpleDelegateGraphTask.ProcessQueryIapResult"), STAT_FSimpleDelegateGraphTask_ProcessQueryIapResult, STATGROUP_TaskGraphTasks);
+
+	FSimpleDelegateGraphTask::CreateAndDispatchWhenReady(
+		FSimpleDelegateGraphTask::FDelegate::CreateLambda([=](){
+			if (IOnlineSubsystem* const OnlineSub = IOnlineSubsystem::Get())
+			{
+				// call store implementation to process query results.
+				if (FOnlineStoreGooglePlay* StoreInterface = (FOnlineStoreGooglePlay*)OnlineSub->GetStoreInterface().Get())
+				{
+					StoreInterface->ProcessQueryAvailablePurchasesResults(bSuccess, ProvidedProductInformation);
+				}
+			}
+			FPlatformMisc::LowLevelOutputDebugStringf(TEXT("In-App Purchase query was completed  %s\n"), bSuccess ? TEXT("successfully") : TEXT("unsuccessfully"));
+		}),
+		GET_STATID(STAT_FSimpleDelegateGraphTask_ProcessQueryIapResult), 
+		nullptr, 
+		ENamedThreads::GameThread
+	);
+
 }
 
 
@@ -138,7 +149,7 @@ void FOnlineStoreGooglePlay::ProcessQueryAvailablePurchasesResults(bool bInSucce
 }
 
 
-bool FOnlineStoreGooglePlay::BeginPurchase(const FString& ProductId, FOnlineInAppPurchaseTransactionRef& InPurchaseStateObject)
+bool FOnlineStoreGooglePlay::BeginPurchase(const FInAppPurchaseProductRequest& ProductRequest, FOnlineInAppPurchaseTransactionRef& InPurchaseStateObject)
 {
 	UE_LOG(LogOnline, Display, TEXT( "FOnlineStoreGooglePlay::BeginPurchase" ));
 	
@@ -149,7 +160,7 @@ bool FOnlineStoreGooglePlay::BeginPurchase(const FString& ProductId, FOnlineInAp
 		CachedPurchaseStateObject = InPurchaseStateObject;
 
 		extern bool AndroidThunkCpp_Iap_BeginPurchase(const FString&, const bool);
-		bCreatedNewTransaction = AndroidThunkCpp_Iap_BeginPurchase(ProductId, true); // TPMB - update begin purchase to use a consumable flag.
+		bCreatedNewTransaction = AndroidThunkCpp_Iap_BeginPurchase(ProductRequest.ProductIdentifier, ProductRequest.bIsConsumable); // TPMB - update begin purchase to use a consumable flag.
 		UE_LOG(LogOnline, Display, TEXT("Created Transaction? - %s"), 
 			bCreatedNewTransaction ? TEXT("Created a transaction.") : TEXT("Failed to create a transaction."));
 
@@ -157,7 +168,7 @@ bool FOnlineStoreGooglePlay::BeginPurchase(const FString& ProductId, FOnlineInAp
 		{
 			UE_LOG(LogOnline, Display, TEXT("FOnlineStoreGooglePlay::BeginPurchase - Could not create a new transaction."));
 			CachedPurchaseStateObject->ReadState = EOnlineAsyncTaskState::Failed;
-			TriggerOnInAppPurchaseCompleteDelegates(EInAppPurchaseState::Invalid, nullptr);
+			TriggerOnInAppPurchaseCompleteDelegates(EInAppPurchaseState::Invalid);
 		}
 		else
 		{
@@ -169,7 +180,7 @@ bool FOnlineStoreGooglePlay::BeginPurchase(const FString& ProductId, FOnlineInAp
 		UE_LOG(LogOnline, Display, TEXT("This device is not able to make purchases."));
 
 		InPurchaseStateObject->ReadState = EOnlineAsyncTaskState::Failed;
-		TriggerOnInAppPurchaseCompleteDelegates(EInAppPurchaseState::NotAllowed, nullptr);
+		TriggerOnInAppPurchaseCompleteDelegates(EInAppPurchaseState::NotAllowed);
 	}
 
 
@@ -179,47 +190,59 @@ bool FOnlineStoreGooglePlay::BeginPurchase(const FString& ProductId, FOnlineInAp
 
 extern "C" void Java_com_epicgames_ue4_GooglePlayStoreHelper_nativePurchaseComplete(JNIEnv* jenv, jobject thiz, jboolean bSuccess, jstring productId, jstring receiptData)
 {
-	if (IOnlineSubsystem* const OnlineSub = IOnlineSubsystem::Get())
+	FString ProductId, ReceiptData;
+	if (bSuccess)
 	{
-		// call store implementation to process query results.
-		if (FOnlineStoreGooglePlay* StoreInterface = (FOnlineStoreGooglePlay*)OnlineSub->GetStoreInterface().Get())
-		{
-			FString ProductId, ReceiptData;
-			if (bSuccess)
-			{
-				const char* charsId = jenv->GetStringUTFChars(productId, 0);
-				ProductId = FString(UTF8_TO_TCHAR(charsId));
-				jenv->ReleaseStringUTFChars(productId, charsId);
+		const char* charsId = jenv->GetStringUTFChars(productId, 0);
+		ProductId = FString(UTF8_TO_TCHAR(charsId));
+		jenv->ReleaseStringUTFChars(productId, charsId);
 
-				const char* charsReceipt = jenv->GetStringUTFChars(receiptData, 0);
-				ReceiptData = FString(UTF8_TO_TCHAR(charsReceipt));
-				jenv->ReleaseStringUTFChars(receiptData, charsReceipt);
-			}
+		const char* charsReceipt = jenv->GetStringUTFChars(receiptData, 0);
+		ReceiptData = FString(UTF8_TO_TCHAR(charsReceipt));
+		jenv->ReleaseStringUTFChars(receiptData, charsReceipt);
+	}
+	FPlatformMisc::LowLevelOutputDebugStringf(TEXT("1... ProductId: %s, ReceiptData: %s\n"), *ProductId, *ReceiptData );
 
-			StoreInterface->ProcessPurchaseResult(bSuccess, ProductId, ReceiptData);
+	DECLARE_CYCLE_STAT(TEXT("FSimpleDelegateGraphTask.ProcessIapResult"), STAT_FSimpleDelegateGraphTask_ProcessIapResult, STATGROUP_TaskGraphTasks);
 
-		}
-	}
+	FSimpleDelegateGraphTask::CreateAndDispatchWhenReady(
+		FSimpleDelegateGraphTask::FDelegate::CreateLambda([=](){
+			FPlatformMisc::LowLevelOutputDebugStringf(TEXT("In-App Purchase was completed  %s\n"), bSuccess ? TEXT("successfully") : TEXT("unsuccessfully"));
+			if (IOnlineSubsystem* const OnlineSub = IOnlineSubsystem::Get())
+			{
+				FPlatformMisc::LowLevelOutputDebugStringf(TEXT("2... ProductId: %s, ReceiptData: %s\n"), *ProductId, *ReceiptData);
+				// call store implementation to process query results.
+				if (FOnlineStoreGooglePlay* StoreInterface = (FOnlineStoreGooglePlay*)OnlineSub->GetStoreInterface().Get())
+				{
+					StoreInterface->ProcessPurchaseResult(bSuccess, ProductId, ReceiptData);
+				}
+			}
+		}),
+		GET_STATID(STAT_FSimpleDelegateGraphTask_ProcessIapResult), 
+		nullptr, 
+		ENamedThreads::GameThread
+	);
 
-	FPlatformMisc::LowLevelOutputDebugStringf(TEXT("In-App Purchase was completed  %s\n"), bSuccess ? TEXT("successfully") : TEXT("unsuccessfully"));
 }
 
 
 void FOnlineStoreGooglePlay::ProcessPurchaseResult(bool bInSuccessful, const FString& ProductId, const FString& InReceiptData)
 {
 	UE_LOG(LogOnline, Display, TEXT("FOnlineStoreGooglePlay::ProcessPurchaseResult"));
+	UE_LOG(LogOnline, Display, TEXT("3... ProductId: %s, ReceiptData: %s\n"), *ProductId, *InReceiptData);
 
-	FGooglePlayPurchaseReceipt Receipt;
-	if (ProductId.Len() > 0 && InReceiptData.Len() > 0)
-	{
-		Receipt.Identifier = ProductId;
-		Receipt.Data = InReceiptData;
-	}
 
 	if (CachedPurchaseStateObject.IsValid())
 	{
+		FInAppPurchaseProductInfo& ProductInfo = CachedPurchaseStateObject->ProvidedProductInformation;
+		ProductInfo.Identifier = ProductId;
+		ProductInfo.DisplayName = TEXT("n/a");
+		ProductInfo.DisplayDescription = TEXT("n/a");
+		ProductInfo.DisplayPrice = TEXT("n/a");
+		ProductInfo.ReceiptData = InReceiptData;
+
 		CachedPurchaseStateObject->ReadState = EOnlineAsyncTaskState::Done;
 	}
 
-	TriggerOnInAppPurchaseCompleteDelegates(bInSuccessful ? EInAppPurchaseState::Success : EInAppPurchaseState::Failed, &Receipt);
+	TriggerOnInAppPurchaseCompleteDelegates(bInSuccessful ? EInAppPurchaseState::Success : EInAppPurchaseState::Failed);
 }
\ No newline at end of file
diff --git a/Engine/Source/Runtime/Online/Android/OnlineSubsystemGooglePlay/Private/OnlineStoreInterfaceGooglePlay.h b/Engine/Source/Runtime/Online/Android/OnlineSubsystemGooglePlay/Private/OnlineStoreInterfaceGooglePlay.h
--- a/Engine/Source/Runtime/Online/Android/OnlineSubsystemGooglePlay/Private/OnlineStoreInterfaceGooglePlay.h
+++ b/Engine/Source/Runtime/Online/Android/OnlineSubsystemGooglePlay/Private/OnlineStoreInterfaceGooglePlay.h
@@ -60,7 +60,7 @@ public:
 
 	// Begin IOnlineStore 
 	virtual bool QueryForAvailablePurchases(const TArray<FString>& ProductIDs, FOnlineProductInformationReadRef& InReadObject) override;
-	virtual bool BeginPurchase(const FString& ProductId, FOnlineInAppPurchaseTransactionRef& InReadObject) override;
+	virtual bool BeginPurchase(const FInAppPurchaseProductRequest& ProductRequest, FOnlineInAppPurchaseTransactionRef& InReadObject) override;
 	virtual bool IsAllowedToMakePurchases() override;
 	// End IOnlineStore 
 
diff --git a/Engine/Source/Runtime/Online/IOS/OnlineSubsystemIOS/Private/OnlineStoreInterfaceIOS.cpp b/Engine/Source/Runtime/Online/IOS/OnlineSubsystemIOS/Private/OnlineStoreInterfaceIOS.cpp
--- a/Engine/Source/Runtime/Online/IOS/OnlineSubsystemIOS/Private/OnlineStoreInterfaceIOS.cpp
+++ b/Engine/Source/Runtime/Online/IOS/OnlineSubsystemIOS/Private/OnlineStoreInterfaceIOS.cpp
@@ -51,42 +51,41 @@
 	UE_LOG(LogOnline, Log, TEXT( "FStoreKitHelper::completeTransaction" ));
 
     
-    [FIOSAsyncTask CreateTaskWithBlock : ^ bool(void)
-    {
-        IOnlineSubsystem* OnlineSub = IOnlineSubsystem::Get(FName(TEXT("IOS")));
-        FOnlineStoreInterfaceIOS* StoreInterface = (FOnlineStoreInterfaceIOS*)OnlineSub->GetStoreInterface().Get();
-     
-        FIOSPurchaseReceipt Receipt;
-        Receipt.Identifier = transaction.transactionIdentifier;
-     
-        if( [IOSAppDelegate GetDelegate].OSVersion >= 7.0 )
-        {
-            NSURL* ReceiptUrl = [[NSBundle mainBundle] appStoreReceiptURL];
-            NSData* ReceiptData = [NSData dataWithContentsOfURL:ReceiptUrl];
-            NSString* nsEncodedReceiptData = [ReceiptData base64EncodedStringWithOptions:NSDataBase64EncodingEndLineWithLineFeed];
-     
-            Receipt.Data = nsEncodedReceiptData;
-        }
-        else
-        {
-            // If earlier than IOS 7, we will need to use the transactionReceipt
-            NSString* nsEncodedReceiptData = [transaction.transactionReceipt base64EncodedStringWithOptions:NSDataBase64EncodingEndLineWithLineFeed];
-            Receipt.Data = nsEncodedReceiptData;
-        }
+	[FIOSAsyncTask CreateTaskWithBlock : ^ bool(void)
+	{
+		IOnlineSubsystem* OnlineSub = IOnlineSubsystem::Get(FName(TEXT("IOS")));
+		FOnlineStoreInterfaceIOS* StoreInterface = (FOnlineStoreInterfaceIOS*)OnlineSub->GetStoreInterface().Get();
+
+		if (StoreInterface->CachedPurchaseStateObject.IsValid())
+		{
+			FString ReceiptData;
+
+			if( [IOSAppDelegate GetDelegate].OSVersion >= 7.0 )
+			{
+				NSURL* nsReceiptUrl = [[NSBundle mainBundle] appStoreReceiptURL];
+				NSData* nsReceiptData = [NSData dataWithContentsOfURL : nsReceiptUrl];
+				NSString* nsEncodedReceiptData = [nsReceiptData base64EncodedStringWithOptions : NSDataBase64EncodingEndLineWithLineFeed];
      
+				ReceiptData = nsEncodedReceiptData;
+			}
+			else
+			{
+				// If earlier than IOS 7, we will need to use the transactionReceipt
+				NSString* nsEncodedReceiptData = [transaction.transactionReceipt base64EncodedStringWithOptions:NSDataBase64EncodingEndLineWithLineFeed];
+				ReceiptData = nsEncodedReceiptData;
+			}
 
-        if (StoreInterface->CachedPurchaseStateObject.IsValid())
-        {
-            StoreInterface->CachedPurchaseStateObject->ReadState = EOnlineAsyncTaskState::Done;
-        }
+			StoreInterface->CachedPurchaseStateObject->ProvidedProductInformation.ReceiptData = ReceiptData;
+			StoreInterface->CachedPurchaseStateObject->ReadState = EOnlineAsyncTaskState::Done;
+		}
      
-        StoreInterface->TriggerOnInAppPurchaseCompleteDelegates(EInAppPurchaseState::Success, &Receipt);
+		StoreInterface->TriggerOnInAppPurchaseCompleteDelegates(EInAppPurchaseState::Success);
      
-        return true;
-    }];
+		return true;
+	}];
  
-    // Remove the transaction from the payment queue.
-    [[SKPaymentQueue defaultQueue] finishTransaction: transaction];
+	// Remove the transaction from the payment queue.
+	[[SKPaymentQueue defaultQueue] finishTransaction: transaction];
 }
 
 
@@ -129,7 +128,7 @@
             StoreInterface->CachedPurchaseStateObject->ReadState = EOnlineAsyncTaskState::Done;
         }
 
-        StoreInterface->TriggerOnInAppPurchaseCompleteDelegates(CompletionState, nullptr);
+        StoreInterface->TriggerOnInAppPurchaseCompleteDelegates(CompletionState);
      
         return true;
      }];
@@ -251,16 +250,18 @@ bool FOnlineStoreInterfaceIOS::IsAllowedToMakePurchases()
 }
 
 
-bool FOnlineStoreInterfaceIOS::BeginPurchase(const FString& ProductId, FOnlineInAppPurchaseTransactionRef& InPurchaseStateObject)
+bool FOnlineStoreInterfaceIOS::BeginPurchase(const FInAppPurchaseProductRequest& ProductRequest, FOnlineInAppPurchaseTransactionRef& InPurchaseStateObject)
 {
 	UE_LOG(LogOnline, Verbose, TEXT( "FOnlineStoreInterfaceIOS::BeginPurchase" ));
 	bool bCreatedNewTransaction = false;
 	
+	const FString& ProductId = ProductRequest.ProductIdentifier;
+
 	if( bIsPurchasing || bIsProductRequestInFlight )
 	{
 		UE_LOG(LogOnline, Verbose, TEXT( "FOnlineStoreInterfaceIOS::BeginPurchase - cannot make a purchase whilst one is in transaction." ));
         
-		TriggerOnInAppPurchaseCompleteDelegates(EInAppPurchaseState::Failed, nullptr);
+		TriggerOnInAppPurchaseCompleteDelegates(EInAppPurchaseState::Failed);
 	}
 	else if( IsAllowedToMakePurchases() )
 	{
@@ -292,7 +293,7 @@ bool FOnlineStoreInterfaceIOS::BeginPurchase(const FString& ProductId, FOnlineIn
 		UE_LOG(LogOnline, Verbose, TEXT("This device is not able to make purchases."));
 		InPurchaseStateObject->ReadState = EOnlineAsyncTaskState::Failed;
         
-        TriggerOnInAppPurchaseCompleteDelegates(EInAppPurchaseState::NotAllowed, nullptr);
+        TriggerOnInAppPurchaseCompleteDelegates(EInAppPurchaseState::NotAllowed);
 	}
 
 	return bCreatedNewTransaction;
diff --git a/Engine/Source/Runtime/Online/IOS/OnlineSubsystemIOS/Private/OnlineStoreInterfaceIOS.h b/Engine/Source/Runtime/Online/IOS/OnlineSubsystemIOS/Private/OnlineStoreInterfaceIOS.h
--- a/Engine/Source/Runtime/Online/IOS/OnlineSubsystemIOS/Private/OnlineStoreInterfaceIOS.h
+++ b/Engine/Source/Runtime/Online/IOS/OnlineSubsystemIOS/Private/OnlineStoreInterfaceIOS.h
@@ -28,20 +28,6 @@ namespace EInAppPurchaseResult
     };
 }
 
-/**
- * IOS implementation of the Platform Purchase receipt. For this we provide an identifier and the encrypted data.
- */
-class FIOSPurchaseReceipt : public IPlatformPurchaseReceipt
-{
-public:
-    // Product identifier
-    FString Identifier;
-    
-    // The encrypted receipt data
-    FString Data;
-};
-
-
 
 /** Helper class, which allows us to manage IAP product information requests, AND transactions */
 @interface FStoreKitHelper : NSObject<SKProductsRequestDelegate, SKPaymentTransactionObserver>
@@ -76,7 +62,7 @@ public:
 
 	// Begin IOnlineStore interface
 	virtual bool QueryForAvailablePurchases(const TArray<FString>& ProductIDs, FOnlineProductInformationReadRef& InReadObject) override;
-	virtual bool BeginPurchase(const FString& ProductId, FOnlineInAppPurchaseTransactionRef& InReadObject) override;
+	virtual bool BeginPurchase(const FInAppPurchaseProductRequest& ProductRequest, FOnlineInAppPurchaseTransactionRef& InReadObject) override;
 	virtual bool IsAllowedToMakePurchases() override;
 	// End IOnlineStore interface
 
diff --git a/Engine/Source/Runtime/Online/OnlineSubsystem/Public/Interfaces/OnlineStoreInterface.h b/Engine/Source/Runtime/Online/OnlineSubsystem/Public/Interfaces/OnlineStoreInterface.h
--- a/Engine/Source/Runtime/Online/OnlineSubsystem/Public/Interfaces/OnlineStoreInterface.h
+++ b/Engine/Source/Runtime/Online/OnlineSubsystem/Public/Interfaces/OnlineStoreInterface.h
@@ -45,10 +45,28 @@ typedef FOnQueryForAvailablePurchasesComplete::FDelegate FOnQueryForAvailablePur
  * @param SessionName the name of the session the that has transitioned to started
  * @param bWasSuccessful true if the async action completed without error, false if there was an error
  */
-DECLARE_MULTICAST_DELEGATE_TwoParams(FOnInAppPurchaseComplete, EInAppPurchaseState::Type, const IPlatformPurchaseReceipt*);
+DECLARE_MULTICAST_DELEGATE_OneParam(FOnInAppPurchaseComplete, EInAppPurchaseState::Type);
 typedef FOnInAppPurchaseComplete::FDelegate FOnInAppPurchaseCompleteDelegate;
 
 
+/**
+ * Micro-transaction request information
+ */
+USTRUCT(BlueprintType)
+struct FInAppPurchaseProductRequest
+{
+	GENERATED_USTRUCT_BODY()
+
+	// The unique product identifier that matches the one from your targeted store.
+	UPROPERTY(BlueprintReadWrite, Category = ProductInfo)
+	FString ProductIdentifier;
+
+	// Flag to determine whether this is a consumable purchase, or not.
+	UPROPERTY(BlueprintReadWrite, Category = ProductInfo)
+	bool bIsConsumable;
+};
+
+
 /**
 * Micro-transaction purchase information
 */
@@ -72,6 +90,10 @@ struct FInAppPurchaseProductInfo
 	// The localized display price name
 	UPROPERTY(BlueprintReadOnly, Category = ProductInfo)
 	FString DisplayPrice;
+
+	// The localized display price name
+	UPROPERTY(BlueprintReadOnly, Category = ProductInfo)
+	FString ReceiptData;
 };
 
 
@@ -149,10 +171,10 @@ public:
 	 *
 	 * @return - whether a purchase request was sent
 	 */
-	virtual bool BeginPurchase(const FString& ProductId, FOnlineInAppPurchaseTransactionRef& InReadObject) = 0;
+	virtual bool BeginPurchase(const FInAppPurchaseProductRequest& ProductRequest, FOnlineInAppPurchaseTransactionRef& InReadObject) = 0;
 	
 	/**
 	 * Delegate which is executed when a Purchase completes
 	 */
-	DEFINE_ONLINE_DELEGATE_TWO_PARAM(OnInAppPurchaseComplete, EInAppPurchaseState::Type, const IPlatformPurchaseReceipt*);
+	DEFINE_ONLINE_DELEGATE_ONE_PARAM(OnInAppPurchaseComplete, EInAppPurchaseState::Type);
 };
\ No newline at end of file
diff --git a/Engine/Source/Runtime/Online/OnlineSubsystemUtils/Classes/InAppPurchaseCallbackProxy.h b/Engine/Source/Runtime/Online/OnlineSubsystemUtils/Classes/InAppPurchaseCallbackProxy.h
--- a/Engine/Source/Runtime/Online/OnlineSubsystemUtils/Classes/InAppPurchaseCallbackProxy.h
+++ b/Engine/Source/Runtime/Online/OnlineSubsystemUtils/Classes/InAppPurchaseCallbackProxy.h
@@ -21,7 +21,7 @@ class UInAppPurchaseCallbackProxy : public UObject
 
 	// Kicks off a transaction for the provided product identifier
 	UFUNCTION(BlueprintCallable, meta = (FriendlyName="Make an In-App Purchase"), Category="Online|InAppPurchase")
-	static UInAppPurchaseCallbackProxy* CreateProxyObjectForInAppPurchase(class APlayerController* PlayerController, const FString& ProductIdentifier);
+	static UInAppPurchaseCallbackProxy* CreateProxyObjectForInAppPurchase(class APlayerController* PlayerController, const FInAppPurchaseProductRequest& ProductRequest);
 
 public:
 
@@ -33,13 +33,13 @@ private:
 
 	/** Called by the InAppPurchase system when the transaction has finished */
 	void OnInAppPurchaseComplete_Delayed();
-	void OnInAppPurchaseComplete(EInAppPurchaseState::Type CompletionState, const IPlatformPurchaseReceipt* ProductReceipt = nullptr);
+	void OnInAppPurchaseComplete(EInAppPurchaseState::Type CompletionState);
 
 	/** Unregisters our delegate from the In-App Purchase system */
 	void RemoveDelegate();
 
 	/** Triggers the In-App Purchase Transaction for the specifed user; the Purchase Request object must already be set up */
-	void Trigger(class APlayerController* PlayerController, const FString& ProductIdentifier);
+	void Trigger(class APlayerController* PlayerController, const FInAppPurchaseProductRequest& ProductRequest);
 
 private:
 
@@ -57,7 +57,4 @@ private:
 
 	/** Did the purchase succeed? */
 	EInAppPurchaseState::Type SavedPurchaseState;
-    
-    /** Copy of the provided purchase receipt */
-    const IPlatformPurchaseReceipt* SavedPurchaseReceipt;
 };
diff --git a/Engine/Source/Runtime/Online/OnlineSubsystemUtils/Private/InAppPurchaseCallbackProxy.cpp b/Engine/Source/Runtime/Online/OnlineSubsystemUtils/Private/InAppPurchaseCallbackProxy.cpp
--- a/Engine/Source/Runtime/Online/OnlineSubsystemUtils/Private/InAppPurchaseCallbackProxy.cpp
+++ b/Engine/Source/Runtime/Online/OnlineSubsystemUtils/Private/InAppPurchaseCallbackProxy.cpp
@@ -10,11 +10,10 @@ UInAppPurchaseCallbackProxy::UInAppPurchaseCallbackProxy(const FObjectInitialize
 {
 	PurchaseRequest = nullptr;
 	WorldPtr = nullptr;
-    SavedPurchaseReceipt = nullptr;
 }
 
 
-void UInAppPurchaseCallbackProxy::Trigger(APlayerController* PlayerController, const FString& ProductIdentifier)
+void UInAppPurchaseCallbackProxy::Trigger(APlayerController* PlayerController, const FInAppPurchaseProductRequest& ProductRequest)
 {
 	bFailedToEvenSubmit = true;
 
@@ -35,7 +34,7 @@ void UInAppPurchaseCallbackProxy::Trigger(APlayerController* PlayerController, c
 				// Set-up, and trigger the transaction through the store interface
 				PurchaseRequest = MakeShareable(new FOnlineInAppPurchaseTransaction());
 				FOnlineInAppPurchaseTransactionRef PurchaseRequestRef = PurchaseRequest.ToSharedRef();
-				StoreInterface->BeginPurchase(ProductIdentifier, PurchaseRequestRef);
+				StoreInterface->BeginPurchase(ProductRequest, PurchaseRequestRef);
 			}
 			else
 			{
@@ -59,11 +58,10 @@ void UInAppPurchaseCallbackProxy::Trigger(APlayerController* PlayerController, c
 }
 
 
-void UInAppPurchaseCallbackProxy::OnInAppPurchaseComplete(EInAppPurchaseState::Type CompletionState, const IPlatformPurchaseReceipt* ProductReceipt)
+void UInAppPurchaseCallbackProxy::OnInAppPurchaseComplete(EInAppPurchaseState::Type CompletionState)
 {
 	RemoveDelegate();
 	SavedPurchaseState = CompletionState;
-    SavedPurchaseReceipt = ProductReceipt;
     
 	if (UWorld* World = WorldPtr.Get())
 	{
@@ -95,7 +93,6 @@ void UInAppPurchaseCallbackProxy::OnInAppPurchaseComplete_Delayed()
 		OnFailure.Broadcast(SavedPurchaseState, ProductInformation);
 	}
     
-    SavedPurchaseReceipt = nullptr;
     PurchaseRequest = nullptr;
 }
 
@@ -125,9 +122,9 @@ void UInAppPurchaseCallbackProxy::BeginDestroy()
 }
 
 
-UInAppPurchaseCallbackProxy* UInAppPurchaseCallbackProxy::CreateProxyObjectForInAppPurchase(class APlayerController* PlayerController, const FString& ProductIdentifier)
+UInAppPurchaseCallbackProxy* UInAppPurchaseCallbackProxy::CreateProxyObjectForInAppPurchase(class APlayerController* PlayerController, const FInAppPurchaseProductRequest& ProductRequest)
 {
 	UInAppPurchaseCallbackProxy* Proxy = NewObject<UInAppPurchaseCallbackProxy>();
-	Proxy->Trigger(PlayerController, ProductIdentifier);
+	Proxy->Trigger(PlayerController, ProductRequest);
 	return Proxy;
 }
