Close movie between queued movies and stop before reset to prevent possible runtime exception. Disable audio for media player plugin.
#android
#ue4

diff --git a/Engine/Build/Android/Java/src/com/epicgames/ue4/MediaPlayer14.java b/Engine/Build/Android/Java/src/com/epicgames/ue4/MediaPlayer14.java
--- a/Engine/Build/Android/Java/src/com/epicgames/ue4/MediaPlayer14.java
+++ b/Engine/Build/Android/Java/src/com/epicgames/ue4/MediaPlayer14.java
@@ -21,15 +21,18 @@ public class MediaPlayer14
 	extends android.media.MediaPlayer
 {
 	private boolean SwizzlePixels = true;
+	private volatile boolean WaitOnBitmapRender = false;
 
 	public MediaPlayer14()
 	{
 		SwizzlePixels = true;
+		WaitOnBitmapRender = false;
 	}
 
 	public MediaPlayer14(boolean swizzlePixels)
 	{
 		SwizzlePixels = swizzlePixels;
+		WaitOnBitmapRender = false;
 	}
 
 	public boolean setDataSource(
@@ -86,6 +89,8 @@ public class MediaPlayer14
 	
 	public void setVideoEnabled(boolean enabled)
 	{
+		WaitOnBitmapRender = true;
+
 		mVideoEnabled = enabled;
 		if (mVideoEnabled && null != mBitmapRenderer.getSurface())
 		{
@@ -95,6 +100,8 @@ public class MediaPlayer14
 		{
 			setSurface(null);
 		}
+
+		WaitOnBitmapRender = false;
 	}
 	
 	public void setAudioEnabled(boolean enabled)
@@ -111,9 +118,13 @@ public class MediaPlayer14
 	
 	public java.nio.Buffer getVideoLastFrameData()
 	{
+
 		if (null != mBitmapRenderer)
 		{
-			return mBitmapRenderer.updateFrameData();
+			WaitOnBitmapRender = true;
+			java.nio.Buffer data = mBitmapRenderer.updateFrameData();
+			WaitOnBitmapRender = false;
+			return data;
 		}
 		else
 		{
@@ -125,7 +136,10 @@ public class MediaPlayer14
 	{
 		if (null != mBitmapRenderer)
 		{
-			return mBitmapRenderer.updateFrameData(destTexture);
+			WaitOnBitmapRender = true;
+			boolean result = mBitmapRenderer.updateFrameData(destTexture);
+			WaitOnBitmapRender = false;
+			return result;
 		}
 		else
 		{
@@ -137,6 +151,7 @@ public class MediaPlayer14
 	{
 		if (null != mBitmapRenderer)
 		{
+			while (WaitOnBitmapRender) ;
 			mBitmapRenderer.release();
 			mBitmapRenderer = null;
 		}
@@ -147,6 +162,7 @@ public class MediaPlayer14
 	{
 		if (null != mBitmapRenderer)
 		{
+			while (WaitOnBitmapRender) ;
 			mBitmapRenderer.release();
 			mBitmapRenderer = null;
 		}
diff --git a/Engine/Plugins/Media/AndroidMedia/Source/AndroidMedia/Private/AndroidMediaPlayer.cpp b/Engine/Plugins/Media/AndroidMedia/Source/AndroidMedia/Private/AndroidMediaPlayer.cpp
--- a/Engine/Plugins/Media/AndroidMedia/Source/AndroidMedia/Private/AndroidMediaPlayer.cpp
+++ b/Engine/Plugins/Media/AndroidMedia/Source/AndroidMedia/Private/AndroidMediaPlayer.cpp
@@ -516,8 +516,9 @@ bool FAndroidMediaPlayer::Open(const FString& Url)
 		if (Extension.Equals(TEXT("3gpp"), ESearchCase::IgnoreCase) ||
 			Extension.Equals(TEXT("mp4"), ESearchCase::IgnoreCase))
 		{
-			// For video we assume both a video and audio track.
-			Tracks.Add(MakeShareable(new AudioTrack(*this, Tracks.Num())));
+			// For video we add video track and disable audio
+			JavaMediaPlayer->SetAudioEnabled(false);
+//			Tracks.Add(MakeShareable(new AudioTrack(*this, Tracks.Num())));
 			Tracks.Add(MakeShareable(new VideoTrack(*this, Tracks.Num())));
 		}
 		else if (Extension.Equals(TEXT("aac"), ESearchCase::IgnoreCase))
diff --git a/Engine/Plugins/Runtime/AndroidMoviePlayer/Source/AndroidMoviePlayer/Private/AndroidMovieStreamer.cpp b/Engine/Plugins/Runtime/AndroidMoviePlayer/Source/AndroidMoviePlayer/Private/AndroidMovieStreamer.cpp
--- a/Engine/Plugins/Runtime/AndroidMoviePlayer/Source/AndroidMoviePlayer/Private/AndroidMovieStreamer.cpp
+++ b/Engine/Plugins/Runtime/AndroidMoviePlayer/Source/AndroidMoviePlayer/Private/AndroidMovieStreamer.cpp
@@ -89,7 +89,7 @@ bool FAndroidMediaPlayerStreamer::Tick(float DeltaTime)
 	{
 		// Current movie finished playing. Clean up current movie and move on to
 		// the next movie, if any.
-		JavaMediaPlayer->Reset();
+		CloseMovie();
 		if (!StartNextMovie())
 		{
 			// Failed to start the next movie, need to indicate that
@@ -195,6 +195,7 @@ bool FAndroidMediaPlayerStreamer::StartNextMovie()
 
 void FAndroidMediaPlayerStreamer::CloseMovie()
 {
+	JavaMediaPlayer->Stop();
 	JavaMediaPlayer->Reset();
 
 	if (Texture.IsValid())
