[UEPLAT-207] Support Message Boxes on Android.

diff --git a/Engine/Build/Android/Java/proguard-project.txt b/Engine/Build/Android/Java/proguard-project.txt
--- a/Engine/Build/Android/Java/proguard-project.txt
+++ b/Engine/Build/Android/Java/proguard-project.txt
@@ -19,6 +19,9 @@
 -keepclassmembers class com.epicgames.ue4.MediaPlayer14 {
    public *;
 }
+-keepclassmembers class com.epicgames.ue4.MessageBox01 {
+   public *;
+}
 
 # Needed...
 -keep class com.epicgames.ue4.GameActivity$JavaAchievement {
diff --git a/Engine/Build/Android/Java/src/com/epicgames/ue4/MessageBox01.java b/Engine/Build/Android/Java/src/com/epicgames/ue4/MessageBox01.java
new file mode 100644
--- /dev/null
+++ b/Engine/Build/Android/Java/src/com/epicgames/ue4/MessageBox01.java
@@ -0,0 +1,185 @@
+// Copyright 1998-2014 Epic Games, Inc. All Rights Reserved.
+
+package com.epicgames.ue4;
+
+import android.app.Dialog;
+import android.app.AlertDialog;
+import android.content.DialogInterface;
+import java.util.ArrayList;
+import java.util.HashMap;
+import android.view.View;
+import android.view.Gravity;
+
+/*
+	Utility class to manage displaying a message box
+	with any number of buttons.
+*/
+class MessageBox01
+	implements View.OnClickListener
+{
+	public MessageBox01()
+	{
+	}
+	
+	public void setCaption(String text)
+	{
+		Caption = text;
+	}
+	
+	public void setText(String text)
+	{
+		Text = text;
+	}
+	
+	public void addButton(String text)
+	{
+		ButtonText.add(text);
+	}
+	
+	public void clear()
+	{
+		Caption = "";
+		Text = "";
+		ButtonText.clear();
+		ButtonPressed = -1;
+	}
+	
+	public int show()
+	{
+		ButtonPressed = -1;
+		createAlert();
+		if (null != Alert)
+		{
+			GameActivity.Get().runOnUiThread(new Runnable()
+				{
+					public void run()
+					{
+						Alert.show();
+					}
+				});
+			while (ButtonPressed == -1 && Alert != null)
+			{
+				try
+				{
+					java.lang.Thread.sleep(100);
+				}
+				catch (java.lang.InterruptedException e)
+				{
+				}
+			}
+		}
+		return ButtonPressed;
+	}
+	
+	protected void createAlert()
+	{
+		dismissAlert();
+		ButtonIdToIndex.clear();
+		Builder = new AlertDialog.Builder(GameActivity.Get());
+		Builder.setTitle(Caption);
+		Builder.setMessage(Text);
+		if (ButtonText.size() >= 1)
+		{
+			Builder.setPositiveButton(ButtonText.get(0), null);
+		}
+		if (ButtonText.size() >= 2)
+		{
+			Builder.setNeutralButton(ButtonText.get(1), null);
+		}
+		Alert = null;
+		GameActivity.Get().runOnUiThread(new Runnable()
+			{
+				public void run()
+				{
+					Alert = Builder.create();
+				}
+			});
+		while (null == Alert)
+		{
+			try
+			{
+				java.lang.Thread.sleep(100);
+			}
+			catch (java.lang.InterruptedException e)
+			{
+			}
+		}
+		Builder = null;
+		Alert.setCancelable(false);
+		Alert.setCanceledOnTouchOutside(false);
+		Alert.setOnShowListener(new DialogInterface.OnShowListener() {
+			public void onShow(DialogInterface dialog)
+			{
+				android.widget.Button button = null;
+				if (ButtonText.size() >= 1)
+				{
+					button = Alert.getButton(
+						android.content.DialogInterface.BUTTON_POSITIVE);
+					registerButton(button, 0);
+				}
+				if (ButtonText.size() >= 2)
+				{
+					button = Alert.getButton(
+						android.content.DialogInterface.BUTTON_NEUTRAL);
+					registerButton(button, 1);
+				}
+				if (button != null)
+				{
+					android.view.ViewGroup parent = (android.view.ViewGroup)button.getParent();
+					for (int i = 2; i < ButtonText.size(); ++i)
+					{
+						button = new android.widget.Button(GameActivity.Get());
+						button.setText(ButtonText.get(i));
+						parent.addView(button);
+						registerButton(button, i);
+					}
+				}
+			}
+			});
+	}
+	
+	protected void registerButton(android.widget.Button button, int button_index)
+	{
+		int new_id = button.getId();
+		if (new_id == View.NO_ID)
+		{
+			new_id = View.generateViewId();
+		}
+		button.setId(new_id);
+		ButtonIdToIndex.put(new_id, button_index);
+		button.setOnClickListener(this);
+		if (ButtonText.size() > 1)
+		{
+			android.widget.LinearLayout.LayoutParams params
+				= (android.widget.LinearLayout.LayoutParams) button.getLayoutParams();
+			params.gravity = Gravity.FILL;
+			params.weight = 1.0f;
+			params.width = android.widget.LinearLayout.LayoutParams.FILL_PARENT;
+			button.setLayoutParams(params);
+		}
+	}
+	
+	public void onClick(android.view.View button)
+	{
+		ButtonPressed = ButtonIdToIndex.get(button.getId());
+		dismissAlert();
+	}
+	
+	protected void dismissAlert()
+	{
+		if (null != Alert)
+		{
+			Alert.dismiss();
+			Alert = null;
+		}
+	}
+	
+	protected String Caption = "";
+	protected String Text = "";
+	protected ArrayList<String> ButtonText
+		= new ArrayList<String>();
+	protected int ButtonPressed = -1;
+	protected AlertDialog Alert = null;
+	protected HashMap<Integer,Integer> ButtonIdToIndex = new HashMap<Integer,Integer>();
+	protected AlertDialog.Builder Builder = null;
+}
diff --git a/Engine/Source/Runtime/Core/Private/Android/AndroidJavaMessageBox.cpp b/Engine/Source/Runtime/Core/Private/Android/AndroidJavaMessageBox.cpp
new file mode 100644
--- /dev/null
+++ b/Engine/Source/Runtime/Core/Private/Android/AndroidJavaMessageBox.cpp
@@ -0,0 +1,57 @@
+// Copyright 1998-2015 Epic Games, Inc. All Rights Reserved.
+
+#include "CorePrivatePCH.h"
+#include "AndroidJavaMessageBox.h"
+
+FJavaAndroidMessageBox::FJavaAndroidMessageBox()
+	: FJavaClassObject(GetClassName(), "()V")
+	, SetCaptionMethod(GetClassMethod("setCaption", "(Ljava/lang/String;)V"))
+	, SetTextMethod(GetClassMethod("setText", "(Ljava/lang/String;)V"))
+	, AddButtonMethod(GetClassMethod("addButton", "(Ljava/lang/String;)V"))
+	, ClearMethod(GetClassMethod("clear", "()V"))
+	, ShowMethod(GetClassMethod("show", "()I"))
+{
+}
+
+void FJavaAndroidMessageBox::SetCaption(const FString & Text)
+{
+	BindObjectToThread();
+	CallMethod<void>(SetCaptionMethod, GetJString(Text));
+}
+
+void FJavaAndroidMessageBox::SetText(const FString & Text)
+{
+	BindObjectToThread();
+	CallMethod<void>(SetTextMethod, GetJString(Text));
+}
+
+void FJavaAndroidMessageBox::AddButton(const FString & Text)
+{
+	BindObjectToThread();
+	CallMethod<void>(AddButtonMethod, GetJString(Text));
+}
+
+void FJavaAndroidMessageBox::Clear()
+{
+	BindObjectToThread();
+	CallMethod<void>(ClearMethod);
+}
+
+int32 FJavaAndroidMessageBox::Show()
+{
+	BindObjectToThread();
+	return CallMethod<int32>(ShowMethod);
+}
+
+FName FJavaAndroidMessageBox::GetClassName()
+{
+	if (FAndroidMisc::GetAndroidBuildVersion() >= 1)
+	{
+		return FName("com/epicgames/ue4/MessageBox01");
+	}
+	else
+	{
+		return FName("");
+	}
+}
+
diff --git a/Engine/Source/Runtime/Core/Private/Android/AndroidMisc.cpp b/Engine/Source/Runtime/Core/Private/Android/AndroidMisc.cpp
--- a/Engine/Source/Runtime/Core/Private/Android/AndroidMisc.cpp
+++ b/Engine/Source/Runtime/Core/Private/Android/AndroidMisc.cpp
@@ -10,6 +10,7 @@
 
 #include "AndroidPlatformCrashContext.h"
 #include "MallocCrash.h"
+#include "AndroidJavaMessageBox.h"
 
 DEFINE_LOG_CATEGORY_STATIC(LogEngine, Log, All);
 
@@ -152,30 +153,83 @@ void FAndroidMisc::ClipboardPaste(class FString& Result)
 
 EAppReturnType::Type FAndroidMisc::MessageBoxExt( EAppMsgType::Type MsgType, const TCHAR* Text, const TCHAR* Caption )
 {
-	//implement android message box here.
-	
-	FPlatformMisc::LowLevelOutputDebugStringf(TEXT("Cannot display dialog box on this platform: %s : %s"), Caption, Text);
-	switch(MsgType)
+	FJavaAndroidMessageBox MessageBox;
+	MessageBox.SetText(Text);
+	MessageBox.SetCaption(Caption);
+	EAppReturnType::Type * ResultValues = nullptr;
+	static EAppReturnType::Type ResultsOk[] = {
+		EAppReturnType::Ok };
+	static EAppReturnType::Type ResultsYesNo[] = {
+		EAppReturnType::Yes, EAppReturnType::No };
+	static EAppReturnType::Type ResultsOkCancel[] = {
+		EAppReturnType::Ok, EAppReturnType::Cancel };
+	static EAppReturnType::Type ResultsYesNoCancel[] = {
+		EAppReturnType::Yes, EAppReturnType::No, EAppReturnType::Cancel };
+	static EAppReturnType::Type ResultsCancelRetryContinue[] = {
+		EAppReturnType::Cancel, EAppReturnType::Retry, EAppReturnType::Continue };
+	static EAppReturnType::Type ResultsYesNoYesAllNoAll[] = {
+		EAppReturnType::Yes, EAppReturnType::No, EAppReturnType::YesAll,
+		EAppReturnType::NoAll };
+	static EAppReturnType::Type ResultsYesNoYesAllNoAllCancel[] = {
+		EAppReturnType::Yes, EAppReturnType::No, EAppReturnType::YesAll,
+		EAppReturnType::NoAll, EAppReturnType::Cancel };
+
+	// TODO: Should we localize button text?
+
+	switch (MsgType)
 	{
 	case EAppMsgType::Ok:
-		return EAppReturnType::Ok; // Ok
+		MessageBox.AddButton(TEXT("Ok"));
+		ResultValues = ResultsOk;
+		break;
 	case EAppMsgType::YesNo:
-		return EAppReturnType::No; // No
+		MessageBox.AddButton(TEXT("Yes"));
+		MessageBox.AddButton(TEXT("No"));
+		ResultValues = ResultsYesNo;
+		break;
 	case EAppMsgType::OkCancel:
-		return EAppReturnType::Cancel; // Cancel
+		MessageBox.AddButton(TEXT("Ok"));
+		MessageBox.AddButton(TEXT("Cancel"));
+		ResultValues = ResultsOkCancel;
+		break;
 	case EAppMsgType::YesNoCancel:
-		return EAppReturnType::Cancel; // Cancel
+		MessageBox.AddButton(TEXT("Yes"));
+		MessageBox.AddButton(TEXT("No"));
+		MessageBox.AddButton(TEXT("Cancel"));
+		ResultValues = ResultsYesNoCancel;
+		break;
 	case EAppMsgType::CancelRetryContinue:
-		return EAppReturnType::Cancel; // Cancel
+		MessageBox.AddButton(TEXT("Cancel"));
+		MessageBox.AddButton(TEXT("Retry"));
+		MessageBox.AddButton(TEXT("Continue"));
+		ResultValues = ResultsCancelRetryContinue;
+		break;
 	case EAppMsgType::YesNoYesAllNoAll:
-		return EAppReturnType::No; // No
+		MessageBox.AddButton(TEXT("Yes"));
+		MessageBox.AddButton(TEXT("No"));
+		MessageBox.AddButton(TEXT("Yes To All"));
+		MessageBox.AddButton(TEXT("No To All"));
+		ResultValues = ResultsYesNoYesAllNoAll;
+		break;
 	case EAppMsgType::YesNoYesAllNoAllCancel:
-		return EAppReturnType::Yes; // Yes
+		MessageBox.AddButton(TEXT("Yes"));
+		MessageBox.AddButton(TEXT("No"));
+		MessageBox.AddButton(TEXT("Yes To All"));
+		MessageBox.AddButton(TEXT("No To All"));
+		MessageBox.AddButton(TEXT("Cancel"));
+		ResultValues = ResultsYesNoYesAllNoAllCancel;
+		break;
 	default:
 		check(0);
 	}
-	return EAppReturnType::Cancel; // Cancel
-
+	int32 Choice = MessageBox.Show();
+	if (Choice >= 0 && nullptr != ResultValues)
+	{
+		return ResultValues[Choice];
+	}
+	// Failed to show dialog, or failed to get a response,
+	// return default cancel response instead.
+	return FGenericPlatformMisc::MessageBoxExt(MsgType, Text, Caption);
 }
 
 extern void AndroidThunkCpp_KeepScreenOn(bool Enable);
diff --git a/Engine/Source/Runtime/Core/Public/Android/AndroidJavaMessageBox.h b/Engine/Source/Runtime/Core/Public/Android/AndroidJavaMessageBox.h
new file mode 100644
--- /dev/null
+++ b/Engine/Source/Runtime/Core/Public/Android/AndroidJavaMessageBox.h
@@ -0,0 +1,27 @@
+// Copyright 1998-2015 Epic Games, Inc. All Rights Reserved.
+
+#pragma once
+
+#include "Core.h"
+#include "AndroidJava.h"
+
+// Wrapper for com/epicgames/ue4/MessageBox*.java.
+class FJavaAndroidMessageBox : public FJavaClassObject
+{
+public:
+	FJavaAndroidMessageBox();
+	void SetCaption(const FString & Text);
+	void SetText(const FString & Text);
+	void AddButton(const FString & Text);
+	void Clear();
+	int32 Show();
+
+private:
+	static FName GetClassName();
+
+	FJavaClassMethod SetCaptionMethod;
+	FJavaClassMethod SetTextMethod;
+	FJavaClassMethod AddButtonMethod;
+	FJavaClassMethod ClearMethod;
+	FJavaClassMethod ShowMethod;
+};
