[UE-6843] Fix graphics streching problems when launchng on a non-keylocked screen-locked device. This moves the screen unlocking from UAT to the game itself. Which allows for proper screen wakeup and lock dismissal.

diff --git a/Engine/Build/Android/Java/AndroidManifest.xml b/Engine/Build/Android/Java/AndroidManifest.xml
--- a/Engine/Build/Android/Java/AndroidManifest.xml
+++ b/Engine/Build/Android/Java/AndroidManifest.xml
@@ -44,4 +44,5 @@
 	<uses-permission android:name="android.permission.GET_ACCOUNTS" />
 	<uses-permission android:name="android.permission.VIBRATE"/>
 	<uses-permission android:name="com.android.vending.BILLING" />
+	<uses-permission android:name="android.permission.DISABLE_KEYGUARD"/>
 </manifest>
diff --git a/Engine/Build/Android/Java/src/com/epicgames/ue4/GameActivity.java b/Engine/Build/Android/Java/src/com/epicgames/ue4/GameActivity.java
--- a/Engine/Build/Android/Java/src/com/epicgames/ue4/GameActivity.java
+++ b/Engine/Build/Android/Java/src/com/epicgames/ue4/GameActivity.java
@@ -39,6 +39,7 @@ import android.view.ViewConfiguration;
 import android.view.ViewGroup.LayoutParams;
 import android.view.ViewGroup.MarginLayoutParams;
 import android.view.WindowManager;
+import android.view.Window;
 import android.widget.LinearLayout;
 import android.widget.PopupWindow;
 
@@ -184,6 +185,18 @@ public class GameActivity extends NativeActivity
 		}
 
 		_activity = this;
+		
+		// Turn on and unlock screen.. Assumption is that this
+		// will only really have an effect when for debug launching
+		// as otherwise the screen is already unlocked.
+		this.getWindow().addFlags(WindowManager.LayoutParams.FLAG_TURN_SCREEN_ON);
+		this.getWindow().addFlags(WindowManager.LayoutParams.FLAG_SHOW_WHEN_LOCKED);
+		this.getWindow().addFlags(WindowManager.LayoutParams.FLAG_DISMISS_KEYGUARD);
+		// On some devices we can also unlock a key-locked screen by disabling the
+		// keylock guard.
+		android.app.KeyguardManager keyman = (android.app.KeyguardManager)getSystemService(KEYGUARD_SERVICE);
+		android.app.KeyguardManager.KeyguardLock keylock = keyman.newKeyguardLock("Unlock");
+		keylock.disableKeyguard();
 
 		// tell Android that we want volume controls to change the media volume, aka music
 		setVolumeControlStream(AudioManager.STREAM_MUSIC);
diff --git a/Engine/Source/Programs/AutomationTool/Android/AndroidPlatform.Automation.cs b/Engine/Source/Programs/AutomationTool/Android/AndroidPlatform.Automation.cs
--- a/Engine/Source/Programs/AutomationTool/Android/AndroidPlatform.Automation.cs
+++ b/Engine/Source/Programs/AutomationTool/Android/AndroidPlatform.Automation.cs
@@ -807,19 +807,14 @@ public class AndroidPlatform : Platform
             throw new AutomationException("Failed to get package name from " + ClientApp);
 		}
 
-		string CommandLine = "shell am start -n " + PackageName + "/com.epicgames.ue4.GameActivity";
-
 		if (Params.Prebuilt)
 		{
 			// clear the log
 			RunAdbCommand(Params, "logcat -c");
 		}
 
-		// Send a command to unlock the device before we try to run it
-		string UnlockCommandLine = "shell input keyevent 82";
-		RunAdbCommand(Params, UnlockCommandLine, null);
-
 		// start the app on device!
+		string CommandLine = "shell am start -n " + PackageName + "/com.epicgames.ue4.GameActivity";
 		ProcessResult ClientProcess = RunAdbCommand(Params, CommandLine, null, ClientRunFlags);
 
 
