[UE-7113] Don't use deprecated API for device unlocking on OS >= 3.2.

diff --git a/Engine/Build/Android/Java/src/com/epicgames/ue4/GameActivity.java b/Engine/Build/Android/Java/src/com/epicgames/ue4/GameActivity.java
--- a/Engine/Build/Android/Java/src/com/epicgames/ue4/GameActivity.java
+++ b/Engine/Build/Android/Java/src/com/epicgames/ue4/GameActivity.java
@@ -189,14 +189,19 @@ public class GameActivity extends NativeActivity
 		// Turn on and unlock screen.. Assumption is that this
 		// will only really have an effect when for debug launching
 		// as otherwise the screen is already unlocked.
-		this.getWindow().addFlags(WindowManager.LayoutParams.FLAG_TURN_SCREEN_ON);
-		this.getWindow().addFlags(WindowManager.LayoutParams.FLAG_SHOW_WHEN_LOCKED);
-		this.getWindow().addFlags(WindowManager.LayoutParams.FLAG_DISMISS_KEYGUARD);
+		this.getWindow().addFlags(
+			WindowManager.LayoutParams.FLAG_TURN_SCREEN_ON |
+			WindowManager.LayoutParams.FLAG_SHOW_WHEN_LOCKED |
+			WindowManager.LayoutParams.FLAG_DISMISS_KEYGUARD);
 		// On some devices we can also unlock a key-locked screen by disabling the
-		// keylock guard.
-		android.app.KeyguardManager keyman = (android.app.KeyguardManager)getSystemService(KEYGUARD_SERVICE);
-		android.app.KeyguardManager.KeyguardLock keylock = keyman.newKeyguardLock("Unlock");
-		keylock.disableKeyguard();
+		// keylock guard. To be safe we only do this on < Android 3.2. As the API
+		// is deprecated from 3.2 onward.
+		if (ANDROID_BUILD_VERSION < 13)
+		{
+			android.app.KeyguardManager keyman = (android.app.KeyguardManager)getSystemService(KEYGUARD_SERVICE);
+			android.app.KeyguardManager.KeyguardLock keylock = keyman.newKeyguardLock("Unlock");
+			keylock.disableKeyguard();
+		}
 
 		// tell Android that we want volume controls to change the media volume, aka music
 		setVolumeControlStream(AudioManager.STREAM_MUSIC);
