Protect against IAP not enabled in onActivityResult
#android
#ue4

diff --git a/Engine/Build/Android/Java/src/com/epicgames/ue4/GameActivity.java b/Engine/Build/Android/Java/src/com/epicgames/ue4/GameActivity.java
--- a/Engine/Build/Android/Java/src/com/epicgames/ue4/GameActivity.java
+++ b/Engine/Build/Android/Java/src/com/epicgames/ue4/GameActivity.java
@@ -861,13 +861,20 @@ public class GameActivity extends NativeActivity
 	@Override
 	protected void onActivityResult(int requestCode, int resultCode, Intent data)
 	{
-		if(!IapStoreHelper.onActivityResult(requestCode, resultCode, data))
+		if( IapStoreHelper != null )
 		{
-			super.onActivityResult(requestCode, resultCode, data);
+			if(!IapStoreHelper.onActivityResult(requestCode, resultCode, data))
+			{
+				super.onActivityResult(requestCode, resultCode, data);
+			}
+			else
+			{
+				Log.debug("[JAVA] - Store Helper handled onActivityResult");
+			}
 		}
 		else
 		{
-			Log.debug("[JAVA] - Store Helper handled onActivityResult");
+			super.onActivityResult(requestCode, resultCode, data);
 		}
 	}
 	
