Fix proguard settings for new classes
#ue4
#android

diff --git a/Engine/Build/Android/Java/proguard-project.txt b/Engine/Build/Android/Java/proguard-project.txt
--- a/Engine/Build/Android/Java/proguard-project.txt
+++ b/Engine/Build/Android/Java/proguard-project.txt
@@ -16,10 +16,10 @@
 -keepclassmembers class com.epicgames.ue4.GameActivity {
    public *;
 }
--keepclassmembers class com.epicgames.ue4.MediaPlayer14 {
+-keep class * extends android.media.MediaPlayer {
    public *;
 }
--keepclassmembers class com.epicgames.ue4.MessageBox01 {
+-keep public class com.epicgames.ue4.MessageBox01 {
    public *;
 }
 
