Fix for MessageBox class in ProGuard file
#ue4
#android

diff --git a/Engine/Build/Android/Java/proguard-project.txt b/Engine/Build/Android/Java/proguard-project.txt
--- a/Engine/Build/Android/Java/proguard-project.txt
+++ b/Engine/Build/Android/Java/proguard-project.txt
@@ -19,7 +19,7 @@
 -keep class * extends android.media.MediaPlayer {
    public *;
 }
--keep public class com.epicgames.ue4.MessageBox01 {
+-keep class com.epicgames.ue4.MessageBox01 {
    public *;
 }
 
