Allow package and project name to be different on Android and convert package name dashes to underscores
#ue4
#android

diff --git a/Engine/Build/Android/Java/src/com/epicgames/ue4/GameActivity.java b/Engine/Build/Android/Java/src/com/epicgames/ue4/GameActivity.java
--- a/Engine/Build/Android/Java/src/com/epicgames/ue4/GameActivity.java
+++ b/Engine/Build/Android/Java/src/com/epicgames/ue4/GameActivity.java
@@ -266,11 +266,15 @@ public class GameActivity extends NativeActivity
 		// Grab a reference to the asset manager
 		AssetManagerReference = this.getAssets();
 
-		// Get the preferred depth buffer size from AndroidManifest.xml
+		// Read metadata from AndroidManifest.xml
 		int DepthBufferPreference = 0;
+		String ProjectName = getPackageName();
+		ProjectName = ProjectName.substring(ProjectName.lastIndexOf('.') + 1);
 		try {
 			ApplicationInfo ai = getPackageManager().getApplicationInfo(getPackageName(), PackageManager.GET_META_DATA);
 			Bundle bundle = ai.metaData;
+
+			// Get the preferred depth buffer size from AndroidManifest.xml
 			if (bundle.containsKey("com.epicgames.ue4.GameActivity.DepthBufferPreference"))
 			{
 				DepthBufferPreference = bundle.getInt("com.epicgames.ue4.GameActivity.DepthBufferPreference");
@@ -281,6 +285,7 @@ public class GameActivity extends NativeActivity
 				Log.debug( "Did not find DepthBufferPreference, using default.");
 			}
 
+			// Determine if data is embedded in APK from AndroidManifest.xml
 			if (bundle.containsKey("com.epicgames.ue4.GameActivity.bPackageDataInsideApk"))
 			{
 				PackageDataInsideApkValue = bundle.getBoolean("com.epicgames.ue4.GameActivity.bPackageDataInsideApk") ? 1 : 0;
@@ -291,6 +296,17 @@ public class GameActivity extends NativeActivity
 				PackageDataInsideApkValue = 0;
 				Log.debug( "Did not find bPackageDataInsideApk, using default.");
 			}
+
+			// Get the project name from AndroidManifest.xml
+			if (bundle.containsKey("com.epicgames.ue4.GameActivity.ProjectName"))
+			{
+				ProjectName = bundle.getString("com.epicgames.ue4.GameActivity.ProjectName");
+				Log.debug( "Found ProjectName = " + ProjectName);
+			}
+			else
+			{
+				Log.debug( "Did not find ProjectName, using package name = " + ProjectName);
+			}
 		}
 		catch (NameNotFoundException e)
 		{
@@ -320,7 +336,7 @@ public class GameActivity extends NativeActivity
 		{
 			int Version = getPackageManager().getPackageInfo(getPackageName(), 0).versionCode;
 			int PatchVersion = 0;
-			nativeSetObbInfo(getApplicationContext().getPackageName(), Version, PatchVersion);
+			nativeSetObbInfo(ProjectName, getApplicationContext().getPackageName(), Version, PatchVersion);
 		}
 		catch (Exception e)
 		{
@@ -952,7 +968,7 @@ public class GameActivity extends NativeActivity
 	public native boolean nativeIsShippingBuild();
 	public native void nativeSetGlobalActivity();
 	public native void nativeSetWindowInfo(boolean bIsPortrait, int DepthBufferPreference);
-	public native void nativeSetObbInfo(String PackageName, int Version, int PatchVersion);
+	public native void nativeSetObbInfo(String ProjectName, String PackageName, int Version, int PatchVersion);
 	public native void nativeSetAndroidVersionInformation( String AndroidVersion, String PhoneMake, String PhoneModel, String OSLanguage );
 
 	public native void nativeConsoleCommand(String commandString);
diff --git a/Engine/Source/Programs/UnrealBuildTool/Android/UEDeployAndroid.cs b/Engine/Source/Programs/UnrealBuildTool/Android/UEDeployAndroid.cs
--- a/Engine/Source/Programs/UnrealBuildTool/Android/UEDeployAndroid.cs
+++ b/Engine/Source/Programs/UnrealBuildTool/Android/UEDeployAndroid.cs
@@ -640,6 +640,7 @@ namespace UnrealBuildTool.Android
 
 			// replace some variables
 			PackageName = PackageName.Replace("[PROJECT]", ProjectName);
+            PackageName = PackageName.Replace("-", "_");
 
 			StringBuilder Text = new StringBuilder();
 			Text.AppendLine("<?xml version=\"1.0\" encoding=\"utf-8\"?>");
@@ -690,7 +691,8 @@ namespace UnrealBuildTool.Android
 			Text.AppendLine("\t\t</activity>");
 			Text.AppendLine(string.Format("\t\t<meta-data android:name=\"com.epicgames.ue4.GameActivity.DepthBufferPreference\" android:value=\"{0}\"/>", ConvertDepthBufferIniValue(DepthBufferPreference)));
 			Text.AppendLine(string.Format("\t\t<meta-data android:name=\"com.epicgames.ue4.GameActivity.bPackageDataInsideApk\" android:value=\"{0}\"/>", bPackageDataInsideApk ? "true" : "false"));
-			Text.AppendLine("\t\t<meta-data android:name=\"com.google.android.gms.games.APP_ID\"");
+            Text.AppendLine(string.Format("\t\t<meta-data android:name=\"com.epicgames.ue4.GameActivity.ProjectName\" android:value=\"{0}\"/>", ProjectName));
+            Text.AppendLine("\t\t<meta-data android:name=\"com.google.android.gms.games.APP_ID\"");
 			Text.AppendLine("\t\t           android:value=\"@string/app_id\" />");
 			Text.AppendLine("\t\t<meta-data android:name=\"com.google.android.gms.version\"");
 			Text.AppendLine("\t\t           android:value=\"@integer/google_play_services_version\" />");
diff --git a/Engine/Source/Runtime/Core/Private/Android/AndroidFile.cpp b/Engine/Source/Runtime/Core/Private/Android/AndroidFile.cpp
--- a/Engine/Source/Runtime/Core/Private/Android/AndroidFile.cpp
+++ b/Engine/Source/Runtime/Core/Private/Android/AndroidFile.cpp
@@ -35,6 +35,7 @@ const FDateTime AndroidEpoch(1970, 1, 1);
 
 
 // AndroidProcess uses this for executable name
+FString GAndroidProjectName;
 FString GPackageName;
 static int32 GPackageVersion = 0;
 static int32 GPackagePatchVersion = 0;
@@ -53,15 +54,18 @@ extern jobject AndroidJNI_GetJavaAssetManager();
 extern AAssetManager * AndroidThunkCpp_GetAssetManager();
 
 //This function is declared in the Java-defined class, GameActivity.java: "public native void nativeSetObbInfo(String PackageName, int Version, int PatchVersion);"
-extern "C" void Java_com_epicgames_ue4_GameActivity_nativeSetObbInfo(JNIEnv* jenv, jobject thiz, jstring PackageName, jint Version, jint PatchVersion)
+extern "C" void Java_com_epicgames_ue4_GameActivity_nativeSetObbInfo(JNIEnv* jenv, jobject thiz, jstring ProjectName, jstring PackageName, jint Version, jint PatchVersion)
 {
-	const char* JavaChars = jenv->GetStringUTFChars(PackageName, 0);
-	GPackageName = UTF8_TO_TCHAR(JavaChars);
+	const char* JavaProjectChars = jenv->GetStringUTFChars(ProjectName, 0);
+	GAndroidProjectName = UTF8_TO_TCHAR(JavaProjectChars);
+	const char* JavaPackageChars = jenv->GetStringUTFChars(PackageName, 0);
+	GPackageName = UTF8_TO_TCHAR(JavaPackageChars);
 	GPackageVersion = Version;
 	GPackagePatchVersion = PatchVersion;
 
-	//Release the string
-	jenv->ReleaseStringUTFChars(PackageName, JavaChars);
+	//Release the strings
+	jenv->ReleaseStringUTFChars(ProjectName, JavaProjectChars);
+	jenv->ReleaseStringUTFChars(PackageName, JavaPackageChars);
 }
 
 // Constructs the base path for any files which are not in OBB/pak data
diff --git a/Engine/Source/Runtime/Core/Private/Android/AndroidProcess.cpp b/Engine/Source/Runtime/Core/Private/Android/AndroidProcess.cpp
--- a/Engine/Source/Runtime/Core/Private/Android/AndroidProcess.cpp
+++ b/Engine/Source/Runtime/Core/Private/Android/AndroidProcess.cpp
@@ -30,19 +30,8 @@ const TCHAR* FAndroidPlatformProcess::BaseDir()
 
 const TCHAR* FAndroidPlatformProcess::ExecutableName(bool bRemoveExtension)
 {
-	static FString CachedExeName;
-	if (CachedExeName.Len() == 0)
-	{
-		extern FString GPackageName;
-		int32 LastDot;
-		GPackageName.FindLastChar('.', LastDot);
-
-		// get just the last bit after all dots
-		CachedExeName = GPackageName.Mid(LastDot + 1);
-	}
-
-	// the string is static, so we can return the characters directly
-	return *CachedExeName;
+	extern FString GAndroidProjectName;
+	return *GAndroidProjectName;
 }
 
 FRunnableThread* FAndroidPlatformProcess::CreateRunnableThread()
