Added another variant of proguard keep explicitly for MediaPlayer14
#ue4
#android

diff --git a/Engine/Build/Android/Java/proguard-project.txt b/Engine/Build/Android/Java/proguard-project.txt
--- a/Engine/Build/Android/Java/proguard-project.txt
+++ b/Engine/Build/Android/Java/proguard-project.txt
@@ -19,6 +19,9 @@
 -keep class * extends android.media.MediaPlayer {
    public *;
 }
+-keep class com.epicgames.ue4.MediaPlayer14 {
+   public *;
+}
 -keep class com.epicgames.ue4.MessageBox01 {
    public *;
 }
