UE-15257 - Android deployment fails on some devices
- Changes deploy script back to old method
- Added code to deal with multiple OBB files for release vs dev and gives the user an option to use one (the other is removed from the device)

diff --git a/Engine/Build/Android/Java/JavaTemplates/DownloaderActivity.java.template b/Engine/Build/Android/Java/JavaTemplates/DownloaderActivity.java.template
--- a/Engine/Build/Android/Java/JavaTemplates/DownloaderActivity.java.template
+++ b/Engine/Build/Android/Java/JavaTemplates/DownloaderActivity.java.template
@@ -29,8 +29,12 @@ import com.google.android.vending.expansion.downloader.IStub;
 
 import android.app.Activity;
 import android.app.PendingIntent;
+import android.app.AlertDialog;
+import android.app.AlertDialog.Builder;
 import android.content.Intent;
 import android.content.pm.PackageManager.NameNotFoundException;
+import android.content.DialogInterface;
+import android.content.DialogInterface.OnClickListener;
 import android.os.AsyncTask;
 import android.os.Bundle;
 import android.os.Messenger;
@@ -84,6 +88,8 @@ public class DownloaderActivity extends Activity implements IDownloaderClient {
     private IDownloaderService mRemoteService;
 
     private IStub mDownloaderClientStub;
+	
+	private final CharSequence[] OBBSelectItems = { "Use Store Data", "Use Development Data" };
 
 	
     private void setState(int newState) {
@@ -100,6 +106,8 @@ public class DownloaderActivity extends Activity implements IDownloaderClient {
         mPauseButton.setText(stringResourceID);
     }
 	
+	static DownloaderActivity _download;
+	
 	private Intent OutputData;
 
     /**
@@ -119,13 +127,30 @@ public class DownloaderActivity extends Activity implements IDownloaderClient {
             String fileName = Helpers.getExpansionAPKFileName(this, xf.mIsMain, xf.mFileVersion);
 			GameActivity.Log.debug("Checking for file : " + fileName);
 			String fileForNewFile = Helpers.generateSaveFileName(this, fileName);
-			GameActivity.Log.debug("which is really being resolved to : " + fileForNewFile);
-            if (!Helpers.doesFileExist(this, fileName, xf.mFileSize, false))
+			String fileForDevFile = Helpers.generateSaveFileNameDevelopment(this, fileName);
+			GameActivity.Log.debug("which is really being resolved to : " + fileForNewFile + "\n Or : " + fileForDevFile);
+            if (!Helpers.doesFileExist(this, fileName, xf.mFileSize, false) &&
+				!Helpers.doesFileExistDev(this, fileName, xf.mFileSize, false))
                 return false;
         }
         return true;
     }
 	
+	boolean onlySingleExpansionFileFound() {
+		for (OBBData.XAPKFile xf : OBBData.xAPKS) {
+            String fileName = Helpers.getExpansionAPKFileName(this, xf.mIsMain, xf.mFileVersion);
+			GameActivity.Log.debug("Checking for file : " + fileName);
+			String fileForNewFile = Helpers.generateSaveFileName(this, fileName);
+			String fileForDevFile = Helpers.generateSaveFileNameDevelopment(this, fileName);
+			
+			if (Helpers.doesFileExist(this, fileName, xf.mFileSize, false) &&
+				Helpers.doesFileExistDev(this, fileName, xf.mFileSize, false))
+                return false;
+		}
+		
+		return true;		
+	}
+	
 	File getFileDetailsCacheFile() {
 		return new File(this.getExternalFilesDir(null), "cacheFile.txt");
 	}
@@ -165,15 +190,51 @@ public class DownloaderActivity extends Activity implements IDownloaderClient {
 		for (OBBData.XAPKFile xf : OBBData.xAPKS) {
             String fileName = Helpers.getExpansionAPKFileName(this, xf.mIsMain, xf.mFileVersion);
 			String fileForNewFile = Helpers.generateSaveFileName(this, fileName);
+			String fileForDevFile = Helpers.generateSaveFileNameDevelopment(this, fileName);
 			// check to see if time/data on files match cached version
 			// if not return false
 			File srcFile = new File(fileForNewFile);
+			File srcDevFile = new File(fileForDevFile);
 			long lastModified = srcFile.lastModified();
-			if(!(fileDetailsMap.containsKey(fileName) && lastModified == fileDetailsMap.get(fileName)))
+			long lastModifiedDev = srcDevFile.lastModified();
+			if(!(srcFile.exists() && fileDetailsMap.containsKey(fileName) && lastModified == fileDetailsMap.get(fileName)) 
+				&&
+			   !(srcDevFile.exists() && fileDetailsMap.containsKey(fileName) && lastModifiedDev == fileDetailsMap.get(fileName)))
 				return false;
 		}
 		return true;
 	}
+	
+	static private void RemoveOBBFile(int OBBToDelete) {
+		
+		for (OBBData.XAPKFile xf : OBBData.xAPKS) {
+		    String fileName = Helpers.getExpansionAPKFileName(DownloaderActivity._download, xf.mIsMain, xf.mFileVersion);
+			switch(OBBToDelete)
+			{
+			case 0:
+				String fileForNewFile = Helpers.generateSaveFileName(DownloaderActivity._download, fileName);
+				File srcFile = new File(fileForNewFile);
+				srcFile.delete();
+				break;
+			case 1:
+				String fileForDevFile = Helpers.generateSaveFileNameDevelopment(DownloaderActivity._download, fileName);
+				File srcDevFile = new File(fileForDevFile);
+				srcDevFile.delete();
+				break;
+			}
+		}		
+	}
+	
+	private void ProcessOBBFiles()
+	{
+		if(GameActivity.Get().VerifyOBBOnStartUp && !expansionFilesUptoData()) {
+				validateXAPKZipFiles();
+		} else {
+				OutputData.putExtra(GameActivity.DOWNLOAD_RETURN_NAME, GameActivity.DOWNLOAD_FILES_PRESENT);		
+				setResult(RESULT_OK, OutputData);
+				finish();
+		}
+	}
 
     /**
      * Calculating a moving average for the validation speed so we don't get
@@ -217,11 +278,21 @@ public class DownloaderActivity extends Activity implements IDownloaderClient {
                     String fileName = Helpers.getExpansionAPKFileName(
                             DownloaderActivity.this,
                             xf.mIsMain, xf.mFileVersion);
-                    if (!Helpers.doesFileExist(DownloaderActivity.this, fileName,
-                            xf.mFileSize, false))
+					boolean normalFile = Helpers.doesFileExist(DownloaderActivity.this, fileName, xf.mFileSize, false);
+					boolean devFile = Helpers.doesFileExistDev(DownloaderActivity.this, fileName, xf.mFileSize, false);
+							
+                    if (!normalFile &&	!devFile )
                         return false;
-                    fileName = Helpers
-                            .generateSaveFileName(DownloaderActivity.this, fileName);
+					
+					if(normalFile)
+                    {
+						fileName = Helpers.generateSaveFileName(DownloaderActivity.this, fileName);
+					}
+					else
+					{
+						fileName = Helpers.generateSaveFileNameDevelopment(DownloaderActivity.this, fileName);
+					}			
+					
                     ZipResourceFile zrf;
                     byte[] buf = new byte[1024 * 256];
                     try {
@@ -329,17 +400,28 @@ public class DownloaderActivity extends Activity implements IDownloaderClient {
 						for (OBBData.XAPKFile xf : OBBData.xAPKS) {
 							String fileName = Helpers.getExpansionAPKFileName(DownloaderActivity.this, xf.mIsMain, xf.mFileVersion);
 							String fileForNewFile = Helpers.generateSaveFileName(DownloaderActivity.this, fileName);
-							
+							String fileForDevFile = Helpers.generateSaveFileNameDevelopment(DownloaderActivity.this, fileName);
+														
 							GameActivity.Log.debug("Writing details for file : " + fileName);
-							
+								
 							File srcFile = new File(fileForNewFile);
-							long lastModified = srcFile.lastModified();
-							bufferedFileCache.write(fileName);
-							bufferedFileCache.write(",");
-							bufferedFileCache.write(new Long(lastModified).toString());
-							bufferedFileCache.newLine();
-							
-							GameActivity.Log.debug("Details for file : " + fileName + " with modified time of " + new Long(lastModified).toString() );
+							File srcDevFile = new File(fileForDevFile);
+							if(srcFile.exists()) {
+								long lastModified = srcFile.lastModified();
+								bufferedFileCache.write(fileName);
+								bufferedFileCache.write(",");
+								bufferedFileCache.write(new Long(lastModified).toString());
+								bufferedFileCache.newLine();
+								GameActivity.Log.debug("Details for file : " + fileName + " with modified time of " + new Long(lastModified).toString() );
+							}	
+							else {
+								long lastModified = srcDevFile.lastModified();
+								bufferedFileCache.write(fileName);
+								bufferedFileCache.write(",");
+								bufferedFileCache.write(new Long(lastModified).toString());
+								bufferedFileCache.newLine();
+								GameActivity.Log.debug("Details for file : " + fileName + " with modified time of " + new Long(lastModified).toString() );
+							}							
 						}
 						
 						bufferedFileCache.close();
@@ -457,7 +539,7 @@ public class DownloaderActivity extends Activity implements IDownloaderClient {
     public void onCreate(Bundle savedInstanceState) {
         super.onCreate(savedInstanceState);
 		GameActivity.Log.debug("Starting DownloaderActivity...");
-		
+		_download = this;
 		// Create somewhere to place the output - we'll check this on 'finish' to make sure we are returning 'something'
 		OutputData = new Intent();
 		
@@ -521,12 +603,25 @@ public class DownloaderActivity extends Activity implements IDownloaderClient {
 
         } else {
 			GameActivity.Log.debug("... Can has! Check 'em Dano!");
-            if(GameActivity.Get().VerifyOBBOnStartUp && !expansionFilesUptoData()) {
-				validateXAPKZipFiles();
-			} else {
-				OutputData.putExtra(GameActivity.DOWNLOAD_RETURN_NAME, GameActivity.DOWNLOAD_FILES_PRESENT);		
-				setResult(RESULT_OK, OutputData);
-				finish();
+			if(!onlySingleExpansionFileFound())	{
+				// Do some UI here to figure out which we want to keep
+				
+				AlertDialog.Builder builder = new AlertDialog.Builder(this);
+				
+				builder.setCancelable(false)
+						.setTitle("Select OBB to use")
+					    .setItems(OBBSelectItems, new DialogInterface.OnClickListener() {
+							public void onClick(DialogInterface dialog, int item) {
+								DownloaderActivity.RemoveOBBFile(item);
+								ProcessOBBFiles();
+							}
+						});
+				
+				AlertDialog alert = builder.create();
+				alert.show();
+			}
+			else {
+				ProcessOBBFiles();
 			}
         }
     }
diff --git a/Engine/Build/Android/Java/src/com/google/android/vending/expansion/downloader/Constants.java b/Engine/Build/Android/Java/src/com/google/android/vending/expansion/downloader/Constants.java
--- a/Engine/Build/Android/Java/src/com/google/android/vending/expansion/downloader/Constants.java
+++ b/Engine/Build/Android/Java/src/com/google/android/vending/expansion/downloader/Constants.java
@@ -32,6 +32,9 @@ public class Constants {
      */
     public static final String EXP_PATH = File.separator + "Android"
             + File.separator + "obb" + File.separator;
+			
+	// Root development location
+	public static final String EXP_PATH_DEV = File.separator + "obb" + File.separator;
     
     /** The intent that gets sent when the service must wake up for a retry */
     public static final String ACTION_RETRY = "android.intent.action.DOWNLOAD_WAKEUP";
diff --git a/Engine/Build/Android/Java/src/com/google/android/vending/expansion/downloader/Helpers.java b/Engine/Build/Android/Java/src/com/google/android/vending/expansion/downloader/Helpers.java
--- a/Engine/Build/Android/Java/src/com/google/android/vending/expansion/downloader/Helpers.java
+++ b/Engine/Build/Android/Java/src/com/google/android/vending/expansion/downloader/Helpers.java
@@ -223,6 +223,18 @@ public class Helpers {
         String path = root.toString() + Constants.EXP_PATH + c.getPackageName();
         return path;
     }
+	
+	static public String generateSaveFileNameDevelopment(Context c, String fileName) {
+        String path = getSaveFilePathDevelopment(c)
+                + File.separator + fileName;
+        return path;
+    }
+
+    static public String getSaveFilePathDevelopment(Context c) {
+        File root = Environment.getExternalStorageDirectory();
+        String path = root.toString() + Constants.EXP_PATH_DEV + c.getPackageName();
+        return path;
+    }
 
     /**
      * Helper function to ascertain the existence of a file and return
@@ -240,6 +252,21 @@ public class Helpers {
         // the file may have been delivered by Market --- let's make sure
         // it's the size we expect
         File fileForNewFile = new File(Helpers.generateSaveFileName(c, fileName));
+        return doesFileExistInternal(c, fileForNewFile, fileSize, deleteFileOnMismatch);
+    }
+	
+	static public boolean doesFileExistDev(Context c, String fileName, long fileSize,
+            boolean deleteFileOnMismatch) {
+        // the file may have been delivered by Market --- let's make sure
+        // it's the size we expect
+        File fileForNewFile = new File(Helpers.generateSaveFileNameDevelopment(c, fileName));
+        return doesFileExistInternal(c, fileForNewFile, fileSize, deleteFileOnMismatch);
+    }
+	
+	static public boolean doesFileExistInternal(Context c, File fileForNewFile, long fileSize,
+            boolean deleteFileOnMismatch) {
+        // the file may have been delivered by Market --- let's make sure
+        // it's the size we expect
         if (fileForNewFile.exists()) {
             if (fileForNewFile.length() == fileSize) {
                 return true;
diff --git a/Engine/Source/Programs/AutomationTool/Android/AndroidPlatform.Automation.cs b/Engine/Source/Programs/AutomationTool/Android/AndroidPlatform.Automation.cs
--- a/Engine/Source/Programs/AutomationTool/Android/AndroidPlatform.Automation.cs
+++ b/Engine/Source/Programs/AutomationTool/Android/AndroidPlatform.Automation.cs
@@ -205,7 +205,6 @@ public class AndroidPlatform : Platform
 						"#!/bin/sh",
 						"cd \"`dirname \"$0\"`\"",
                         "ADB=",
-                        "PREFIX=Android/",
 						"if [ \"$ANDROID_HOME\" != \"\" ]; then ADB=$ANDROID_HOME/platform-tools/adb; else ADB=" +Environment.GetEnvironmentVariable("ANDROID_HOME") + "; fi",
 						"DEVICE=",
 						"if [ \"$1\" != \"\" ]; then DEVICE=\"-s $1\"; fi",
@@ -217,15 +216,14 @@ public class AndroidPlatform : Platform
 						"$ADB $DEVICE install " + Path.GetFileName(ApkName),
 						"if [ $? -eq 0 ]; then",
 						"\techo",
-                        "\tNEWVERSION=$(echo \"`$ADB $DEVICE shell 'getprop ro.build.version.release'`\" | awk -F\".\" '{gsub(/-.*/,\"\",$4); if ($1 > 4 || ($1 == 4 && $2 >= 2 )) print 1; else print 0 }')",
-                        "\t if [ $NEWVERSION -eq 1 ]; then STORAGE=/mnt/shell/emulated; PREFIX=; else STORAGE=$(echo \"`$ADB $DEVICE shell 'echo $EXTERNAL_STORAGE'`\" | cat -v | tr -d '^M'); fi",
 						"\techo Removing old data. Failures here are usually fine - indicating the files were not on the device.",
-                        "\t$ADB $DEVICE shell rm -r $STORAGE/UE4Game/" + Params.ShortProjectName,
-						"\t$ADB $DEVICE shell rm -r $STORAGE/UE4Game/UE4CommandLine.txt",
-						"\t$ADB $DEVICE shell rm -r $STORAGE/" + TargetAndroidLocation + PackageName,
+                        "\t$ADB $DEVICE shell 'rm -r $EXTERNAL_STORAGE/UE4Game/" + Params.ShortProjectName + "'",
+						"\t$ADB $DEVICE shell 'rm -r $EXTERNAL_STORAGE/UE4Game/UE4CommandLine.txt" + "'",
+						"\t$ADB $DEVICE shell 'rm -r $EXTERNAL_STORAGE/" + TargetAndroidLocation + PackageName + "'",
 						bPackageDataInsideApk ? "" : "\techo",
 						bPackageDataInsideApk ? "" : "\techo Installing new data. Failures here indicate storage problems \\(missing SD card or bad permissions\\) and are fatal.",
-						bPackageDataInsideApk ? "" : "\t$ADB $DEVICE push " + Path.GetFileName(ObbName) + " $STORAGE/$PREFIX" + DeviceObbName,
+						bPackageDataInsideApk ? "" : "\tSTORAGE=$(echo \"`$ADB $DEVICE shell 'echo $EXTERNAL_STORAGE'`\" | cat -v | tr -d '^M')",
+						bPackageDataInsideApk ? "" : "\t$ADB $DEVICE push " + Path.GetFileName(ObbName) + " $STORAGE/" + DeviceObbName,
 						bPackageDataInsideApk ? "if [ 1 ]; then" : "\tif [ $? -eq 0 ]; then",
 						"\t\techo",
 						"\t\techo Installation successful",
@@ -251,26 +249,8 @@ public class AndroidPlatform : Platform
                         "if \"%ANDROIDHOME%\"==\"\" set ANDROIDHOME="+Environment.GetEnvironmentVariable("ANDROID_HOME"),
 						"set ADB=%ANDROIDHOME%\\platform-tools\\adb.exe",
 						"set DEVICE=",
-                        "set PREFIX=Android/",
                         "if not \"%1\"==\"\" set DEVICE=-s %1",
                         "for /f \"delims=\" %%A in ('%ADB% %DEVICE% " + GetStorageQueryCommand() +"') do @set STORAGE=%%A",
-                        "for /f \"tokens=1-2 delims=.\" %%M in ('%ADB% %DEVICE% shell getprop ro.build.version.release') do (",
-                        "    if \"%%M\"==\"\" (goto completed)",
-                        "    if \"%%N\"==\"\" (goto completed)",
-                        "    if %%M equ 4 (",     
-                        "        if %%N geq 2 (",
-                        "            set STORAGE=/mnt/shell/emulated",
-                        "            set PREFIX=",
-                        "        );",
-                        "    ) else (",
-                        "       if %%M gtr 4 (",
-                        "            set STORAGE=/mnt/shell/emulated",
-                        "            set PREFIX=",
-                        "       );",
-                        "    );",                        
-                        "    @echo.",
-                        ")",
-                        ":completed",
 						"@echo.",
 						"@echo Uninstalling existing application. Failures here can almost always be ignored.",
 						"%ADB% %DEVICE% uninstall " + PackageName,
@@ -283,7 +263,7 @@ public class AndroidPlatform : Platform
 						"%ADB% %DEVICE% shell rm -r %STORAGE%/" + TargetAndroidLocation + PackageName,
 						bPackageDataInsideApk ? "" : "@echo.",
 						bPackageDataInsideApk ? "" : "@echo Installing new data. Failures here indicate storage problems (missing SD card or bad permissions) and are fatal.",
-						bPackageDataInsideApk ? "" : "%ADB% %DEVICE% push " + Path.GetFileName(ObbName) + " %STORAGE%/%PREFIX%" + DeviceObbName,
+						bPackageDataInsideApk ? "" : "%ADB% %DEVICE% push " + Path.GetFileName(ObbName) + " %STORAGE%/" + DeviceObbName,
 						bPackageDataInsideApk ? "" : "if \"%ERRORLEVEL%\" NEQ \"0\" goto Error",
 						"@echo.",
 						"@echo Installation successful",
