UE-15347 - Application will crash when key missing for Android Store and OBB download needed
- Changed download activity so that it checks for a key length > 0 before trying to download and throws up a dialogue box if a key isn't found.

diff --git a/Engine/Build/Android/Java/JavaTemplates/DownloaderActivity.java.template b/Engine/Build/Android/Java/JavaTemplates/DownloaderActivity.java.template
--- a/Engine/Build/Android/Java/JavaTemplates/DownloaderActivity.java.template
+++ b/Engine/Build/Android/Java/JavaTemplates/DownloaderActivity.java.template
@@ -558,42 +558,65 @@ public class DownloaderActivity extends Activity implements IDownloaderClient {
 				GameActivity.Log.debug("... Whoops... missing; go go go download system!");
 				
             try {
-                Intent launchIntent = DownloaderActivity.this
-                        .getIntent();
-                Intent intentToLaunchThisActivityFromNotification = new Intent(
-                        DownloaderActivity
-                        .this, DownloaderActivity.this.getClass());
-                intentToLaunchThisActivityFromNotification.setFlags(Intent.FLAG_ACTIVITY_NEW_TASK |
-                        Intent.FLAG_ACTIVITY_CLEAR_TOP);
-                intentToLaunchThisActivityFromNotification.setAction(launchIntent.getAction());
-
-                if (launchIntent.getCategories() != null) {
-                    for (String category : launchIntent.getCategories()) {
-                        intentToLaunchThisActivityFromNotification.addCategory(category);
-                    }
-                }
-
-                // Build PendingIntent used to open this activity from
-                // Notification
-                PendingIntent pendingIntent = PendingIntent.getActivity(
-                        DownloaderActivity.this,
-                        0, intentToLaunchThisActivityFromNotification,
-                        PendingIntent.FLAG_UPDATE_CURRENT);
-                // Request to start the download
-                int startResult = DownloaderClientMarshaller.startDownloadServiceIfRequired(this,
-                        pendingIntent, OBBDownloaderService.class);
-
-                if (startResult != DownloaderClientMarshaller.NO_DOWNLOAD_REQUIRED) {
-                    // The DownloaderService has started downloading the files,
-                    // show progress
-                    initializeDownloadUI();
-                    return;
-                } // otherwise, download not needed so we fall through to saying all is OK
+			
+				// Make sure we have a key before we try to start the service
+				if(OBBDownloaderService.getPublicKeyLength() == 0) {
+					AlertDialog.Builder builder = new AlertDialog.Builder(this);
+				
+					builder.setCancelable(false)
+							.setTitle("No Google Play Store Key")
+							.setMessage("No OBB found and no store key to try to download. Please set one up in Android Project Settings")
+							.setPositiveButton("Exit", new DialogInterface.OnClickListener() {
+								public void onClick(DialogInterface dialog, int item) {
+									OutputData.putExtra(GameActivity.DOWNLOAD_RETURN_NAME, GameActivity.DOWNLOAD_NO_PLAY_KEY);
+									setResult(RESULT_OK, OutputData);
+									finish();
+								}
+							});
+					
+					AlertDialog alert = builder.create();
+					alert.show();
+				}
 				else
 				{
-					OutputData.putExtra(GameActivity.DOWNLOAD_RETURN_NAME, GameActivity.DOWNLOAD_FILES_PRESENT);
-					setResult(RESULT_OK, OutputData);
-					finish();
+			
+					Intent launchIntent = DownloaderActivity.this
+							.getIntent();
+					Intent intentToLaunchThisActivityFromNotification = new Intent(
+							DownloaderActivity
+							.this, DownloaderActivity.this.getClass());
+					intentToLaunchThisActivityFromNotification.setFlags(Intent.FLAG_ACTIVITY_NEW_TASK |
+							Intent.FLAG_ACTIVITY_CLEAR_TOP);
+					intentToLaunchThisActivityFromNotification.setAction(launchIntent.getAction());
+
+					if (launchIntent.getCategories() != null) {
+						for (String category : launchIntent.getCategories()) {
+							intentToLaunchThisActivityFromNotification.addCategory(category);
+						}
+					}
+
+					// Build PendingIntent used to open this activity from
+					// Notification
+					PendingIntent pendingIntent = PendingIntent.getActivity(
+							DownloaderActivity.this,
+							0, intentToLaunchThisActivityFromNotification,
+							PendingIntent.FLAG_UPDATE_CURRENT);
+					// Request to start the download
+					int startResult = DownloaderClientMarshaller.startDownloadServiceIfRequired(this,
+							pendingIntent, OBBDownloaderService.class);
+
+					if (startResult != DownloaderClientMarshaller.NO_DOWNLOAD_REQUIRED) {
+						// The DownloaderService has started downloading the files,
+						// show progress
+						initializeDownloadUI();
+						return;
+					} // otherwise, download not needed so we fall through to saying all is OK
+					else
+					{
+						OutputData.putExtra(GameActivity.DOWNLOAD_RETURN_NAME, GameActivity.DOWNLOAD_FILES_PRESENT);
+						setResult(RESULT_OK, OutputData);
+						finish();
+					}
 				}
                   
             } catch (NameNotFoundException e) {
diff --git a/Engine/Build/Android/Java/JavaTemplates/OBBDownloaderService.java.template b/Engine/Build/Android/Java/JavaTemplates/OBBDownloaderService.java.template
--- a/Engine/Build/Android/Java/JavaTemplates/OBBDownloaderService.java.template
+++ b/Engine/Build/Android/Java/JavaTemplates/OBBDownloaderService.java.template
@@ -31,6 +31,10 @@ public class OBBDownloaderService extends DownloaderService {
             -100, -12, 43, 2, -8, -4, 9, 5, -106, -108, -33, 45, -1, 84
     };
 
+	public static int getPublicKeyLength() {
+		return BASE64_PUBLIC_KEY.length();
+	}
+	
     /**
      * This public key comes from your Android Market publisher account, and it
      * used by the LVL to validate responses from Market on your behalf.
diff --git a/Engine/Build/Android/Java/src/com/epicgames/ue4/GameActivity.java b/Engine/Build/Android/Java/src/com/epicgames/ue4/GameActivity.java
--- a/Engine/Build/Android/Java/src/com/epicgames/ue4/GameActivity.java
+++ b/Engine/Build/Android/Java/src/com/epicgames/ue4/GameActivity.java
@@ -89,6 +89,7 @@ public class GameActivity extends NativeActivity
 	public static final int DOWNLOAD_USER_QUIT = 3;    // user aborted the download
 	public static final int DOWNLOAD_FAILED = 4;
 	public static final int DOWNLOAD_INVALID = 5;
+	public static final int DOWNLOAD_NO_PLAY_KEY = 6;
 	public static final String DOWNLOAD_RETURN_NAME = "Result";
 	
 	static GameActivity _activity;
@@ -1038,6 +1039,9 @@ public class GameActivity extends NativeActivity
 			case DOWNLOAD_INVALID:
 				logMsg += "Download Invalid";
 				break;
+			case DOWNLOAD_NO_PLAY_KEY:
+				logMsg +="Download No Play Key";
+				break;
 			default:
 				logMsg += "Unknown message!";
 				break;
@@ -1050,7 +1054,8 @@ public class GameActivity extends NativeActivity
 			if(errorCode == DOWNLOAD_NO_RETURN_CODE 
 			|| errorCode == DOWNLOAD_USER_QUIT 
 			|| errorCode == DOWNLOAD_FAILED 
-			|| errorCode == DOWNLOAD_INVALID)
+			|| errorCode == DOWNLOAD_INVALID
+			|| errorCode == DOWNLOAD_NO_PLAY_KEY)
 				finish();
 		}
 		else if( IapStoreHelper != null )
