Android media player fixes
#ue-11850
#ue4
#android

diff --git a/Engine/Build/Android/Java/src/com/epicgames/ue4/MediaPlayer14.java b/Engine/Build/Android/Java/src/com/epicgames/ue4/MediaPlayer14.java
--- a/Engine/Build/Android/Java/src/com/epicgames/ue4/MediaPlayer14.java
+++ b/Engine/Build/Android/Java/src/com/epicgames/ue4/MediaPlayer14.java
@@ -12,6 +12,8 @@ import android.os.Build;
 import java.io.File;
 import java.io.IOException;
 import java.io.RandomAccessFile;
+import android.media.MediaPlayer;
+import java.util.Random;
 
 /*
 	Custom media player that renders video to a frame buffer. This
@@ -27,6 +29,17 @@ public class MediaPlayer14
 	{
 		SwizzlePixels = true;
 		WaitOnBitmapRender = false;
+
+		 setOnErrorListener(new OnErrorListener()
+		 {
+			@Override
+			public boolean onError(    MediaPlayer mp,    int what,    int extra)
+			{
+			  GameActivity.Log.warn("Error occurred while playing audio. ("+what+", "+extra+")");
+			  return true;
+			}
+		}
+		);
 	}
 
 	public MediaPlayer14(boolean swizzlePixels)
@@ -50,11 +63,7 @@ public class MediaPlayer14
 			}
 			RandomAccessFile data = new RandomAccessFile(f, "r");
 			setDataSource(data.getFD(), offset, size);
-			if (!CreateBitmapRenderer())
-			{
-				reset();
-				return false;
-			}
+			releaseBitmapRenderer();
 		}
 		catch(IOException e)
 		{
@@ -72,11 +81,7 @@ public class MediaPlayer14
 		{
 			AssetFileDescriptor assetFD = assetManager.openFd(assetPath);
 			setDataSource(assetFD.getFileDescriptor(), offset, size);
-			if (!CreateBitmapRenderer())
-			{
-				reset();
-				return false;
-			}
+			releaseBitmapRenderer();
 		}
 		catch(IOException e)
 		{
@@ -116,9 +121,24 @@ public class MediaPlayer14
 		}
 	}
 	
+	public void initBitmapRenderer()
+	{
+		// if not already allocated.
+		// Create bitmap renderer's gl resources in the renderer thread.
+		  if (null == mBitmapRenderer)
+		  {
+			if (!CreateBitmapRenderer())
+			{
+				GameActivity.Log.warn("initBitmapRenderer failed to alloc mBitmapRenderer ");
+				reset();
+			  }
+		  }
+	}
+
 	public java.nio.Buffer getVideoLastFrameData()
 	{
 
+		initBitmapRenderer();
 		if (null != mBitmapRenderer)
 		{
 			WaitOnBitmapRender = true;
@@ -134,6 +154,7 @@ public class MediaPlayer14
 
 	public boolean getVideoLastFrame(int destTexture)
 	{
+		initBitmapRenderer();
 		if (null != mBitmapRenderer)
 		{
 			WaitOnBitmapRender = true;
@@ -152,8 +173,7 @@ public class MediaPlayer14
 		if (null != mBitmapRenderer)
 		{
 			while (WaitOnBitmapRender) ;
-			mBitmapRenderer.release();
-			mBitmapRenderer = null;
+			releaseBitmapRenderer();
 		}
 		super.release();
 	}
@@ -163,8 +183,7 @@ public class MediaPlayer14
 		if (null != mBitmapRenderer)
 		{
 			while (WaitOnBitmapRender) ;
-			mBitmapRenderer.release();
-			mBitmapRenderer = null;
+			releaseBitmapRenderer();
 		}
 		super.reset();
 	}
@@ -255,6 +274,7 @@ public class MediaPlayer14
             }
             mPositionAttrib = GLES20.glGetAttribLocation(mProgram, "Position");
             mTexCoordsAttrib = GLES20.glGetAttribLocation(mProgram, "TexCoords");
+			mTextureUniform = GLES20.glGetUniformLocation(mProgram, "VideoTexture");
 
 			// Create blit mesh.
             mTriangleVertices = java.nio.ByteBuffer.allocateDirect(
@@ -380,6 +400,7 @@ public class MediaPlayer14
 		private int mPositionAttrib;
 		private int mTexCoordsAttrib;
 		private int mBlitBuffer;
+		private int mTextureUniform;
 
 		public java.nio.Buffer updateFrameData()
 		{
@@ -525,6 +546,13 @@ public class MediaPlayer14
 				GLES20.GL_COLOR_ATTACHMENT0,
 				GLES20.GL_TEXTURE_2D, FBOTextureID, 0);
 
+			// check status
+			int status = GLES20.glCheckFramebufferStatus(GLES20.GL_FRAMEBUFFER);
+			if (status != GLES20.GL_FRAMEBUFFER_COMPLETE)
+			{
+				GameActivity.Log.warn("Failed to complete framebuffer attachment ("+status+")");
+			}
+
 			// The special shaders to render from the video texture.
 			GLES20.glUseProgram(mProgram);
 
@@ -540,6 +568,13 @@ public class MediaPlayer14
 
 			glVerify("setup movie texture read");
 
+            GLES20.glClear( GLES20.GL_COLOR_BUFFER_BIT);
+
+			// connect 'VideoTexture' to video source texture (mTextureID) in texture unit 0.
+			GLES20.glActiveTexture(GLES20.GL_TEXTURE0);
+			GLES20.glBindTexture(GL_TEXTURE_EXTERNAL_OES, mTextureID);
+			GLES20.glUniform1i(mTextureUniform, 0);
+
 			// Draw the video texture mesh.
 			GLES20.glDrawArrays(GLES20.GL_TRIANGLE_STRIP, 0, 4);
 
@@ -647,17 +682,18 @@ public class MediaPlayer14
 
 	private boolean CreateBitmapRenderer()
 	{
-		if (null != mBitmapRenderer)
-		{
-			mBitmapRenderer.release();
-			mBitmapRenderer = null;
-		}
+		releaseBitmapRenderer();
+
 		mBitmapRenderer = new BitmapRenderer();
 		if (!mBitmapRenderer.isValid())
 		{
 			mBitmapRenderer = null;
 			return false;
 		}
+		
+		// set this here as the size may have been set before the GL resources were created.
+		mBitmapRenderer.setSize(getVideoWidth(),getVideoHeight());
+
 		setOnVideoSizeChangedListener(new android.media.MediaPlayer.OnVideoSizeChangedListener() {
 			public void onVideoSizeChanged(android.media.MediaPlayer player, int w, int h)
 			{
@@ -669,5 +705,16 @@ public class MediaPlayer14
 		return true;
 	}
 
+	void releaseBitmapRenderer()
+	{
+		if (null != mBitmapRenderer)
+		{
+			mBitmapRenderer.release();
+			mBitmapRenderer = null;
+			setSurface(null);
+			setOnVideoSizeChangedListener(null);
+		}
+	}
+
 	private BitmapRenderer mBitmapRenderer = null;
 }
