Invalidate window resolution cache on resume from background on Android
#ue-14790
#ue4
#android

diff --git a/Engine/Build/Android/Java/src/com/epicgames/ue4/GameActivity.java b/Engine/Build/Android/Java/src/com/epicgames/ue4/GameActivity.java
--- a/Engine/Build/Android/Java/src/com/epicgames/ue4/GameActivity.java
+++ b/Engine/Build/Android/Java/src/com/epicgames/ue4/GameActivity.java
@@ -110,6 +110,9 @@ public class GameActivity extends NativeActivity
 	private static int PackageDataInsideApkValue = -1;
 	private static int HasOBBFiles = -1;
 	
+	// depthbuffer preference from manifest
+	int DepthBufferPreference = 0;
+	
 	/** AssetManger reference - populated on start up and used when the OBB is packed into the APK */
 	private AssetManager			AssetManagerReference;
 	
@@ -294,7 +297,6 @@ public class GameActivity extends NativeActivity
 		AssetManagerReference = this.getAssets();
 
 		// Read metadata from AndroidManifest.xml
-		int DepthBufferPreference = 0;
 		String ProjectName = getPackageName();
 		ProjectName = ProjectName.substring(ProjectName.lastIndexOf('.') + 1);
 		try {
@@ -552,6 +554,9 @@ public class GameActivity extends NativeActivity
 	public void onResume()
 	{
 		super.onResume();
+
+		// invalidate window cache
+		nativeSetWindowInfo(getResources().getConfiguration().orientation == Configuration.ORIENTATION_PORTRAIT, DepthBufferPreference);
 		
 		// only do this on KitKat and above
 		if (android.os.Build.VERSION.SDK_INT >= 19 && ShouldHideUI)
diff --git a/Engine/Source/Runtime/Core/Private/Android/AndroidWindow.cpp b/Engine/Source/Runtime/Core/Private/Android/AndroidWindow.cpp
--- a/Engine/Source/Runtime/Core/Private/Android/AndroidWindow.cpp
+++ b/Engine/Source/Runtime/Core/Private/Android/AndroidWindow.cpp
@@ -61,6 +61,7 @@ static bool GAndroidIsPortrait = false;
 static int GAndroidDepthBufferPreference = 0;
 extern "C" void Java_com_epicgames_ue4_GameActivity_nativeSetWindowInfo(JNIEnv* jenv, jobject thiz, jboolean bIsPortrait, jint DepthBufferPreference)
 {
+	WindowInit = false;
 	GAndroidIsPortrait = bIsPortrait == JNI_TRUE;
 	GAndroidDepthBufferPreference = DepthBufferPreference;
 	FPlatformMisc::LowLevelOutputDebugStringf(TEXT("App is running in %s\n"), GAndroidIsPortrait ? TEXT("Portrait") : TEXT("Landscape"));
