UEPLAT-454 - Add option for fullscreen immersive mode for Android projects
- support added for Kitkat for fullscreen sticky mode

diff --git a/Engine/Build/Android/Java/src/com/epicgames/ue4/GameActivity.java b/Engine/Build/Android/Java/src/com/epicgames/ue4/GameActivity.java
--- a/Engine/Build/Android/Java/src/com/epicgames/ue4/GameActivity.java
+++ b/Engine/Build/Android/Java/src/com/epicgames/ue4/GameActivity.java
@@ -559,10 +559,46 @@ public class GameActivity extends NativeActivity
 		nativeSetWindowInfo(getResources().getConfiguration().orientation == Configuration.ORIENTATION_PORTRAIT, DepthBufferPreference);
 		
 		// only do this on KitKat and above
-		if (android.os.Build.VERSION.SDK_INT >= 19 && ShouldHideUI)
+		if (ShouldHideUI)
 		{ 
+		
 			View decorView = getWindow().getDecorView(); 
-			decorView.setSystemUiVisibility(View.SYSTEM_UI_FLAG_LAYOUT_STABLE | View.SYSTEM_UI_FLAG_LAYOUT_HIDE_NAVIGATION | View.SYSTEM_UI_FLAG_LAYOUT_FULLSCREEN | View.SYSTEM_UI_FLAG_HIDE_NAVIGATION | View.SYSTEM_UI_FLAG_FULLSCREEN | View.SYSTEM_UI_FLAG_IMMERSIVE_STICKY); 
+			if(android.os.Build.VERSION.SDK_INT >= 19) {
+				decorView.setSystemUiVisibility(View.SYSTEM_UI_FLAG_LAYOUT_STABLE
+											| View.SYSTEM_UI_FLAG_LAYOUT_HIDE_NAVIGATION
+											| View.SYSTEM_UI_FLAG_LAYOUT_FULLSCREEN
+											| View.SYSTEM_UI_FLAG_HIDE_NAVIGATION
+											| View.SYSTEM_UI_FLAG_FULLSCREEN
+											| View.SYSTEM_UI_FLAG_IMMERSIVE_STICKY); 
+			} /*else {
+				decorView.setSystemUiVisibility(View.SYSTEM_UI_FLAG_LAYOUT_STABLE
+											| View.SYSTEM_UI_FLAG_LAYOUT_HIDE_NAVIGATION
+											| View.SYSTEM_UI_FLAG_LAYOUT_FULLSCREEN
+											| View.SYSTEM_UI_FLAG_HIDE_NAVIGATION
+											| View.SYSTEM_UI_FLAG_FULLSCREEN);
+
+			}*/
+			
+			decorView.setOnSystemUiVisibilityChangeListener(new View.OnSystemUiVisibilityChangeListener() {
+				@Override public void onSystemUiVisibilityChange(int visibility) {
+							View decorView = getWindow().getDecorView(); 
+							if(android.os.Build.VERSION.SDK_INT >= 19) {
+								decorView.setSystemUiVisibility(View.SYSTEM_UI_FLAG_LAYOUT_STABLE
+															| View.SYSTEM_UI_FLAG_LAYOUT_HIDE_NAVIGATION
+															| View.SYSTEM_UI_FLAG_LAYOUT_FULLSCREEN
+															| View.SYSTEM_UI_FLAG_HIDE_NAVIGATION
+															| View.SYSTEM_UI_FLAG_FULLSCREEN
+															| View.SYSTEM_UI_FLAG_IMMERSIVE_STICKY); 
+							} /*else {
+								decorView.setSystemUiVisibility(View.SYSTEM_UI_FLAG_LAYOUT_STABLE
+															| View.SYSTEM_UI_FLAG_LAYOUT_HIDE_NAVIGATION
+															| View.SYSTEM_UI_FLAG_LAYOUT_FULLSCREEN
+															| View.SYSTEM_UI_FLAG_HIDE_NAVIGATION
+															| View.SYSTEM_UI_FLAG_FULLSCREEN);
+
+							}*/
+						}
+			});
 		}
 		
 		if(HasAllFiles)
diff --git a/Engine/Source/Programs/UnrealBuildTool/Android/UEDeployAndroid.cs b/Engine/Source/Programs/UnrealBuildTool/Android/UEDeployAndroid.cs
--- a/Engine/Source/Programs/UnrealBuildTool/Android/UEDeployAndroid.cs
+++ b/Engine/Source/Programs/UnrealBuildTool/Android/UEDeployAndroid.cs
@@ -922,6 +922,8 @@ namespace UnrealBuildTool.Android
 			Ini.GetString("/Script/AndroidRuntimeSettings.AndroidRuntimeSettings", "VersionDisplayName", out VersionDisplayName);
 			string Orientation;
 			Ini.GetString("/Script/AndroidRuntimeSettings.AndroidRuntimeSettings", "Orientation", out Orientation);
+            bool EnableFullScreen;
+            Ini.GetBool("/Script/AndroidRuntimeSettings.AndroidRuntimeSettings", "bFullScreen", out EnableFullScreen);
 			List<string> ExtraManifestNodeTags;
 			Ini.GetArray("/Script/AndroidRuntimeSettings.AndroidRuntimeSettings", "ExtraManifestNodeTags", out ExtraManifestNodeTags);
 			List<string> ExtraApplicationNodeTags;
@@ -1023,6 +1025,7 @@ namespace UnrealBuildTool.Android
 			Text.AppendLine(string.Format("\t\t<meta-data android:name=\"com.epicgames.ue4.GameActivity.DepthBufferPreference\" android:value=\"{0}\"/>", ConvertDepthBufferIniValue(DepthBufferPreference)));
 			Text.AppendLine(string.Format("\t\t<meta-data android:name=\"com.epicgames.ue4.GameActivity.bPackageDataInsideApk\" android:value=\"{0}\"/>", bPackageDataInsideApk ? "true" : "false"));
             Text.AppendLine(string.Format("\t\t<meta-data android:name=\"com.epicgames.ue4.GameActivity.bVerifyOBBOnStartUp\" android:value=\"{0}\"/>", (bIsForDistribution && !bDisableVerifyOBBOnStartUp) ? "true" : "false"));
+            Text.AppendLine(string.Format("\t\t<meta-data android:name=\"com.epicgames.ue4.GameActivity.bShouldHideUI=\" android:value=\"{0}\"/>", EnableFullScreen ? "true" : "false"));
             Text.AppendLine(string.Format("\t\t<meta-data android:name=\"com.epicgames.ue4.GameActivity.ProjectName\" android:value=\"{0}\"/>", ProjectName));
 			Text.AppendLine(string.Format("\t\t<meta-data android:name=\"com.epicgames.ue4.GameActivity.bHasOBBFiles\" android:value=\"{0}\"/>", bHasOBBFiles ? "true" : "false"));
             Text.AppendLine("\t\t<meta-data android:name=\"com.google.android.gms.games.APP_ID\"");
diff --git a/Engine/Source/Runtime/Android/AndroidRuntimeSettings/Classes/AndroidRuntimeSettings.h b/Engine/Source/Runtime/Android/AndroidRuntimeSettings/Classes/AndroidRuntimeSettings.h
--- a/Engine/Source/Runtime/Android/AndroidRuntimeSettings/Classes/AndroidRuntimeSettings.h
+++ b/Engine/Source/Runtime/Android/AndroidRuntimeSettings/Classes/AndroidRuntimeSettings.h
@@ -136,6 +136,10 @@ public:
 	UPROPERTY(GlobalConfig, EditAnywhere, Category = APKPackaging)
 	TEnumAsByte<EAndroidScreenOrientation::Type> Orientation;
 
+	// Should the software navigation buttons be hidden or not
+	UPROPERTY(GlobalConfig, EditAnywhere, Category = APKPackaging, Meta = (DisplayName = "Enable FullScreen Immersive on KitKat and above devices."))
+	bool bFullScreen;
+
 	// The preferred depth buffer bitcount for Android
 	UPROPERTY(GlobalConfig, EditAnywhere, Category = APKPackaging, Meta = (DisplayName = "Preferred Depth Buffer format"))
 	TEnumAsByte<EAndroidDepthBufferPreference::Type> DepthBufferPreference;
