Fix Android packaging on newest SDK releases. Google removed paths to several tools in build.xml.
#ue4
#android

diff --git a/Engine/Build/Android/Java/JavaLibs/downloader_library/custom_rules.xml b/Engine/Build/Android/Java/JavaLibs/downloader_library/custom_rules.xml
new file mode 100644
--- /dev/null
+++ b/Engine/Build/Android/Java/JavaLibs/downloader_library/custom_rules.xml
@@ -0,0 +1,20 @@
+<?xml version="1.0" encoding="utf-8"?>
+<project name="custom_rules" default="debug">
+    <fail
+            message="sdk.dir is missing. Make sure to generate local.properties using 'android update project' or to inject it through the ANDROID_HOME environment variable."
+            unless="sdk.dir"
+    />
+    
+	<path id="android.antlibs">
+        <pathelement path="${sdk.dir}/tools/lib/ant-tasks.jar" />
+    </path>
+	
+	<!-- Custom tasks -->
+    <taskdef resource="anttasks.properties" classpathref="android.antlibs" />
+	
+	<getbuildtools name="android.build.tools.dir" verbose="${verbose}" />
+    <property name="aapt" location="${android.build.tools.dir}/aapt" />
+    <property name="aidl" location="${android.build.tools.dir}/aidl" />
+    <property name="dx" location="${android.build.tools.dir}/dx.bat" />
+    <property name="zipalign" location="${android.build.tools.dir}/zipalign" />
+</project>
diff --git a/Engine/Build/Android/Java/JavaLibs/google-play-services_lib_rev19/custom_rules.xml b/Engine/Build/Android/Java/JavaLibs/google-play-services_lib_rev19/custom_rules.xml
new file mode 100644
--- /dev/null
+++ b/Engine/Build/Android/Java/JavaLibs/google-play-services_lib_rev19/custom_rules.xml
@@ -0,0 +1,20 @@
+<?xml version="1.0" encoding="utf-8"?>
+<project name="custom_rules" default="debug">
+    <fail
+            message="sdk.dir is missing. Make sure to generate local.properties using 'android update project' or to inject it through the ANDROID_HOME environment variable."
+            unless="sdk.dir"
+    />
+    
+	<path id="android.antlibs">
+        <pathelement path="${sdk.dir}/tools/lib/ant-tasks.jar" />
+    </path>
+	
+	<!-- Custom tasks -->
+    <taskdef resource="anttasks.properties" classpathref="android.antlibs" />
+	
+	<getbuildtools name="android.build.tools.dir" verbose="${verbose}" />
+    <property name="aapt" location="${android.build.tools.dir}/aapt" />
+    <property name="aidl" location="${android.build.tools.dir}/aidl" />
+    <property name="dx" location="${android.build.tools.dir}/dx.bat" />
+    <property name="zipalign" location="${android.build.tools.dir}/zipalign" />
+</project>
diff --git a/Engine/Build/Android/Java/custom_rules.xml b/Engine/Build/Android/Java/custom_rules.xml
new file mode 100644
--- /dev/null
+++ b/Engine/Build/Android/Java/custom_rules.xml
@@ -0,0 +1,20 @@
+<?xml version="1.0" encoding="utf-8"?>
+<project name="custom_rules" default="debug">
+    <fail
+            message="sdk.dir is missing. Make sure to generate local.properties using 'android update project' or to inject it through the ANDROID_HOME environment variable."
+            unless="sdk.dir"
+    />
+    
+	<path id="android.antlibs">
+        <pathelement path="${sdk.dir}/tools/lib/ant-tasks.jar" />
+    </path>
+	
+	<!-- Custom tasks -->
+    <taskdef resource="anttasks.properties" classpathref="android.antlibs" />
+	
+	<getbuildtools name="android.build.tools.dir" verbose="${verbose}" />
+    <property name="aapt" location="${android.build.tools.dir}/aapt" />
+    <property name="aidl" location="${android.build.tools.dir}/aidl" />
+    <property name="dx" location="${android.build.tools.dir}/dx.bat" />
+    <property name="zipalign" location="${android.build.tools.dir}/zipalign" />
+</project>
