Fix older Android SDK zipalign path for Ant and an issue for mac (.bat was hardcoded)
#ue4
#android

diff --git a/Engine/Build/Android/Java/JavaLibs/downloader_library/custom_rules.xml b/Engine/Build/Android/Java/JavaLibs/downloader_library/custom_rules.xml
--- a/Engine/Build/Android/Java/JavaLibs/downloader_library/custom_rules.xml
+++ b/Engine/Build/Android/Java/JavaLibs/downloader_library/custom_rules.xml
@@ -1,20 +1,37 @@
 <?xml version="1.0" encoding="utf-8"?>
 <project name="custom_rules" default="debug">
-    <fail
-            message="sdk.dir is missing. Make sure to generate local.properties using 'android update project' or to inject it through the ANDROID_HOME environment variable."
-            unless="sdk.dir"
-    />
+	<fail
+		message="sdk.dir is missing. Make sure to generate local.properties using 'android update project' or to inject it through the ANDROID_HOME environment variable."
+		unless="sdk.dir"
+	/>
     
 	<path id="android.antlibs">
-        <pathelement path="${sdk.dir}/tools/lib/ant-tasks.jar" />
-    </path>
+		<pathelement path="${sdk.dir}/tools/lib/ant-tasks.jar" />
+	</path>
 	
 	<!-- Custom tasks -->
-    <taskdef resource="anttasks.properties" classpathref="android.antlibs" />
+	<taskdef resource="anttasks.properties" classpathref="android.antlibs" />
 	
+	<condition property="exe" value=".exe" else=""><os family="windows" /></condition>
+	<condition property="bat" value=".bat" else=""><os family="windows" /></condition>
+
 	<getbuildtools name="android.build.tools.dir" verbose="${verbose}" />
-    <property name="aapt" location="${android.build.tools.dir}/aapt" />
-    <property name="aidl" location="${android.build.tools.dir}/aidl" />
-    <property name="dx" location="${android.build.tools.dir}/dx.bat" />
-    <property name="zipalign" location="${android.build.tools.dir}/zipalign" />
+	<property name="aapt" location="${android.build.tools.dir}/aapt${exe}" />
+	<property name="aidl" location="${android.build.tools.dir}/aidl${exe}" />
+	<property name="dx" location="${android.build.tools.dir}/dx${bat}" />
+	<if>
+        	<condition>
+          		<resourceexists>
+              			<file file="${android.build.tools.dir}/zipalign${exe}" />
+			</resourceexists>
+		</condition>
+		<then>
+			<property name="zipalign" location="${android.build.tools.dir}/zipalign${exe}" />
+		</then>
+		<else>
+			<property name="android.tools.dir" location="${sdk.dir}/tools" />
+			<property name="zipalign" location="${android.tools.dir}/zipalign${exe}" />
+		</else>
+	</if>
+
 </project>
diff --git a/Engine/Build/Android/Java/JavaLibs/google-play-services_lib_rev19/custom_rules.xml b/Engine/Build/Android/Java/JavaLibs/google-play-services_lib_rev19/custom_rules.xml
--- a/Engine/Build/Android/Java/JavaLibs/google-play-services_lib_rev19/custom_rules.xml
+++ b/Engine/Build/Android/Java/JavaLibs/google-play-services_lib_rev19/custom_rules.xml
@@ -1,20 +1,37 @@
 <?xml version="1.0" encoding="utf-8"?>
 <project name="custom_rules" default="debug">
-    <fail
-            message="sdk.dir is missing. Make sure to generate local.properties using 'android update project' or to inject it through the ANDROID_HOME environment variable."
-            unless="sdk.dir"
-    />
+	<fail
+		message="sdk.dir is missing. Make sure to generate local.properties using 'android update project' or to inject it through the ANDROID_HOME environment variable."
+		unless="sdk.dir"
+	/>
     
 	<path id="android.antlibs">
-        <pathelement path="${sdk.dir}/tools/lib/ant-tasks.jar" />
-    </path>
+		<pathelement path="${sdk.dir}/tools/lib/ant-tasks.jar" />
+	</path>
 	
 	<!-- Custom tasks -->
-    <taskdef resource="anttasks.properties" classpathref="android.antlibs" />
+	<taskdef resource="anttasks.properties" classpathref="android.antlibs" />
 	
+	<condition property="exe" value=".exe" else=""><os family="windows" /></condition>
+	<condition property="bat" value=".bat" else=""><os family="windows" /></condition>
+
 	<getbuildtools name="android.build.tools.dir" verbose="${verbose}" />
-    <property name="aapt" location="${android.build.tools.dir}/aapt" />
-    <property name="aidl" location="${android.build.tools.dir}/aidl" />
-    <property name="dx" location="${android.build.tools.dir}/dx.bat" />
-    <property name="zipalign" location="${android.build.tools.dir}/zipalign" />
+	<property name="aapt" location="${android.build.tools.dir}/aapt${exe}" />
+	<property name="aidl" location="${android.build.tools.dir}/aidl${exe}" />
+	<property name="dx" location="${android.build.tools.dir}/dx${bat}" />
+	<if>
+        	<condition>
+          		<resourceexists>
+              			<file file="${android.build.tools.dir}/zipalign${exe}" />
+			</resourceexists>
+		</condition>
+		<then>
+			<property name="zipalign" location="${android.build.tools.dir}/zipalign${exe}" />
+		</then>
+		<else>
+			<property name="android.tools.dir" location="${sdk.dir}/tools" />
+			<property name="zipalign" location="${android.tools.dir}/zipalign${exe}" />
+		</else>
+	</if>
+
 </project>
diff --git a/Engine/Build/Android/Java/custom_rules.xml b/Engine/Build/Android/Java/custom_rules.xml
--- a/Engine/Build/Android/Java/custom_rules.xml
+++ b/Engine/Build/Android/Java/custom_rules.xml
@@ -1,20 +1,37 @@
 <?xml version="1.0" encoding="utf-8"?>
 <project name="custom_rules" default="debug">
-    <fail
-            message="sdk.dir is missing. Make sure to generate local.properties using 'android update project' or to inject it through the ANDROID_HOME environment variable."
-            unless="sdk.dir"
-    />
+	<fail
+		message="sdk.dir is missing. Make sure to generate local.properties using 'android update project' or to inject it through the ANDROID_HOME environment variable."
+		unless="sdk.dir"
+	/>
     
 	<path id="android.antlibs">
-        <pathelement path="${sdk.dir}/tools/lib/ant-tasks.jar" />
-    </path>
+		<pathelement path="${sdk.dir}/tools/lib/ant-tasks.jar" />
+	</path>
 	
 	<!-- Custom tasks -->
-    <taskdef resource="anttasks.properties" classpathref="android.antlibs" />
+	<taskdef resource="anttasks.properties" classpathref="android.antlibs" />
 	
+	<condition property="exe" value=".exe" else=""><os family="windows" /></condition>
+	<condition property="bat" value=".bat" else=""><os family="windows" /></condition>
+
 	<getbuildtools name="android.build.tools.dir" verbose="${verbose}" />
-    <property name="aapt" location="${android.build.tools.dir}/aapt" />
-    <property name="aidl" location="${android.build.tools.dir}/aidl" />
-    <property name="dx" location="${android.build.tools.dir}/dx.bat" />
-    <property name="zipalign" location="${android.build.tools.dir}/zipalign" />
+	<property name="aapt" location="${android.build.tools.dir}/aapt${exe}" />
+	<property name="aidl" location="${android.build.tools.dir}/aidl${exe}" />
+	<property name="dx" location="${android.build.tools.dir}/dx${bat}" />
+	<if>
+        	<condition>
+          		<resourceexists>
+              			<file file="${android.build.tools.dir}/zipalign${exe}" />
+			</resourceexists>
+		</condition>
+		<then>
+			<property name="zipalign" location="${android.build.tools.dir}/zipalign${exe}" />
+		</then>
+		<else>
+			<property name="android.tools.dir" location="${sdk.dir}/tools" />
+			<property name="zipalign" location="${android.tools.dir}/zipalign${exe}" />
+		</else>
+	</if>
+
 </project>
