Remove libGLES3 dependency from libvrapi
#ue4
#android

diff --git a/Engine/Build/Android/Java/src/com/epicgames/ue4/GameActivity.java b/Engine/Build/Android/Java/src/com/epicgames/ue4/GameActivity.java
--- a/Engine/Build/Android/Java/src/com/epicgames/ue4/GameActivity.java
+++ b/Engine/Build/Android/Java/src/com/epicgames/ue4/GameActivity.java
@@ -1198,6 +1198,14 @@ public class GameActivity extends NativeActivity
 	static
 	{
 		System.loadLibrary("gnustl_shared");
+		try
+		{
+			System.loadLibrary("vrapi");
+		}
+		catch (java.lang.UnsatisfiedLinkError e)
+		{
+			Log.debug("GearVR library not loaded. Ignore this if GearVR plugin intentionally disabled.");
+		}
 		System.loadLibrary("UE4");
 	}
 }
