#gearvr Adding entitlement checking for GearVR

diff --git a/Engine/Build/Android/Java/src/com/epicgames/ue4/GameActivity.java b/Engine/Build/Android/Java/src/com/epicgames/ue4/GameActivity.java
--- a/Engine/Build/Android/Java/src/com/epicgames/ue4/GameActivity.java
+++ b/Engine/Build/Android/Java/src/com/epicgames/ue4/GameActivity.java
@@ -374,6 +374,8 @@ public class GameActivity extends NativeActivity
 				PackagedForGearVR = true;
 				String VRMode = bundle.getString("com.samsung.android.vr.application.mode");
 				Log.debug("Found GearVR mode = " + VRMode);
+
+				com.oculus.svclib.OVREntitlementChecker.doAutomatedCheck(this);
 			}
 			else
 			{
