Android broadcast console cmd

diff --git a/Engine/Build/Android/Java/src/com/epicgames/ue4/ConsoleCmdReceiver.java b/Engine/Build/Android/Java/src/com/epicgames/ue4/ConsoleCmdReceiver.java
new file mode 100644
--- /dev/null
+++ b/Engine/Build/Android/Java/src/com/epicgames/ue4/ConsoleCmdReceiver.java
@@ -0,0 +1,45 @@
+/*
+ * Copyright (C) 2012 The Android Open Source Project
+ *
+ * Licensed under the Apache License, Version 2.0 (the "License");
+ * you may not use this file except in compliance with the License.
+ * You may obtain a copy of the License at
+ *
+ *      http://www.apache.org/licenses/LICENSE-2.0
+ *
+ * Unless required by applicable law or agreed to in writing, software
+ * distributed under the License is distributed on an "AS IS" BASIS,
+ * WITHOUT WARRANTIES OR CONDITIONS OF ANY KIND, either express or implied.
+ * See the License for the specific language governing permissions and
+ * limitations under the License.
+ */
+
+package com.epicgames.ue4;
+
+import android.content.BroadcastReceiver;
+import android.content.Context;
+import android.content.Intent;
+
+public class ConsoleCmdReceiver extends BroadcastReceiver {
+
+	private GameActivity gameActivity;
+	
+	// Constructor...
+	public ConsoleCmdReceiver(GameActivity InGameActivity) {
+		gameActivity = InGameActivity;
+	}
+		
+	@Override
+    public void onReceive(Context context, Intent intent) {
+			// example usage
+			// adb shell am broadcast -a android.intent.action.RUN -e cmd "stat fps"
+			String action = intent.getAction();
+			if (action != null && action == Intent.ACTION_RUN) {
+				String cmd = intent.getStringExtra("cmd");
+				if (cmd != null) {
+					gameActivity.nativeConsoleCommand(cmd);
+				}
+			}
+    }
+
+}
diff --git a/Engine/Build/Android/Java/src/com/epicgames/ue4/GameActivity.java b/Engine/Build/Android/Java/src/com/epicgames/ue4/GameActivity.java
--- a/Engine/Build/Android/Java/src/com/epicgames/ue4/GameActivity.java
+++ b/Engine/Build/Android/Java/src/com/epicgames/ue4/GameActivity.java
@@ -21,6 +21,7 @@ import android.text.InputType;
 import android.content.Context;
 import android.content.DialogInterface;
 import android.content.Intent;
+import android.content.IntentFilter;
 import android.content.res.AssetManager;
 import android.content.res.Configuration;
 import android.content.IntentSender.SendIntentException;
@@ -62,6 +63,8 @@ import java.net.HttpURLConnection;
 import com.epicgames.ue4.GooglePlayStoreHelper;
 import com.epicgames.ue4.GooglePlayLicensing;
 
+//
+import com.epicgames.ue4.ConsoleCmdReceiver;
 
 // TODO: use the resources from the UE4 lib project once we've got the packager up and running
 //import com.epicgames.ue4.R;
@@ -106,6 +109,9 @@ public class GameActivity extends NativeActivity
 	AlertDialog virtualKeyboardAlert;
 	EditText virtualKeyboardInputBox;
 
+	// Console commands receiver
+	ConsoleCmdReceiver consoleCmdReceiver;
+
 	// default the PackageDataInsideApk to an invalid value to make sure we don't get it too early
 	private static int PackageDataInsideApkValue = -1;
 	private static int HasOBBFiles = -1;
@@ -157,6 +163,9 @@ public class GameActivity extends NativeActivity
 	private boolean InitCompletedOK = false;
 	
 	private boolean ShouldHideUI = false;
+
+	/** Whether this application is for distribution */
+	private boolean IsForDistribution = false;
 	
 	/** Access singleton activity for game. **/
 	public static GameActivity Get()
@@ -179,6 +188,13 @@ public class GameActivity extends NativeActivity
 	{
 		super.onStart();
 		
+		if (!IsForDistribution)
+		{
+			// Create console command broadcast listener
+			consoleCmdReceiver = new ConsoleCmdReceiver(this);
+			registerReceiver(consoleCmdReceiver, new IntentFilter(Intent.ACTION_RUN));
+		}
+
 		Log.debug("==================================> Inside onStart function in GameActivity");
 	}
 
@@ -303,6 +319,11 @@ public class GameActivity extends NativeActivity
 			ApplicationInfo ai = getPackageManager().getApplicationInfo(getPackageName(), PackageManager.GET_META_DATA);
 			Bundle bundle = ai.metaData;
 
+			if ((ai.flags & ApplicationInfo.FLAG_DEBUGGABLE) == 0) 
+			{
+				IsForDistribution = true;
+			}
+			
 			// Get the preferred depth buffer size from AndroidManifest.xml
 			if (bundle.containsKey("com.epicgames.ue4.GameActivity.DepthBufferPreference"))
 			{
@@ -624,6 +645,11 @@ public class GameActivity extends NativeActivity
 	public void onStop()
 	{
 		super.onStop();
+
+		if (consoleCmdReceiver != null)
+		{
+			unregisterReceiver(consoleCmdReceiver);
+		}
 	}
 
 	// handle ad popup visibility and requests
