Back out changelist 2665783 (too early)

diff --git a/Engine/Build/Android/Java/src/com/epicgames/ue4/ConsoleCmdReceiver.java b/Engine/Build/Android/Java/src/com/epicgames/ue4/ConsoleCmdReceiver.java
deleted file mode 100644
diff --git a/Engine/Build/Android/Java/src/com/epicgames/ue4/GameActivity.java b/Engine/Build/Android/Java/src/com/epicgames/ue4/GameActivity.java
--- a/Engine/Build/Android/Java/src/com/epicgames/ue4/GameActivity.java
+++ b/Engine/Build/Android/Java/src/com/epicgames/ue4/GameActivity.java
@@ -21,7 +21,6 @@ import android.text.InputType;
 import android.content.Context;
 import android.content.DialogInterface;
 import android.content.Intent;
-import android.content.IntentFilter;
 import android.content.res.AssetManager;
 import android.content.res.Configuration;
 import android.content.IntentSender.SendIntentException;
@@ -63,8 +62,6 @@ import java.net.HttpURLConnection;
 import com.epicgames.ue4.GooglePlayStoreHelper;
 import com.epicgames.ue4.GooglePlayLicensing;
 
-//
-import com.epicgames.ue4.ConsoleCmdReceiver;
 
 // TODO: use the resources from the UE4 lib project once we've got the packager up and running
 //import com.epicgames.ue4.R;
@@ -109,9 +106,6 @@ public class GameActivity extends NativeActivity
 	AlertDialog virtualKeyboardAlert;
 	EditText virtualKeyboardInputBox;
 
-	// Console commands receiver
-	ConsoleCmdReceiver consoleCmdReceiver;
-
 	// default the PackageDataInsideApk to an invalid value to make sure we don't get it too early
 	private static int PackageDataInsideApkValue = -1;
 	private static int HasOBBFiles = -1;
@@ -163,9 +157,6 @@ public class GameActivity extends NativeActivity
 	private boolean InitCompletedOK = false;
 	
 	private boolean ShouldHideUI = false;
-
-	/** Whether this application is for distribution */
-	private boolean IsForDistribution = false;
 	
 	/** Access singleton activity for game. **/
 	public static GameActivity Get()
@@ -188,13 +179,6 @@ public class GameActivity extends NativeActivity
 	{
 		super.onStart();
 		
-		if (!IsForDistribution)
-		{
-			// Create console command broadcast listener
-			consoleCmdReceiver = new ConsoleCmdReceiver(this);
-			registerReceiver(consoleCmdReceiver, new IntentFilter(Intent.ACTION_RUN));
-		}
-
 		Log.debug("==================================> Inside onStart function in GameActivity");
 	}
 
@@ -319,11 +303,6 @@ public class GameActivity extends NativeActivity
 			ApplicationInfo ai = getPackageManager().getApplicationInfo(getPackageName(), PackageManager.GET_META_DATA);
 			Bundle bundle = ai.metaData;
 
-			if ((ai.flags & ApplicationInfo.FLAG_DEBUGGABLE) == 0) 
-			{
-				IsForDistribution = true;
-			}
-			
 			// Get the preferred depth buffer size from AndroidManifest.xml
 			if (bundle.containsKey("com.epicgames.ue4.GameActivity.DepthBufferPreference"))
 			{
@@ -645,11 +624,6 @@ public class GameActivity extends NativeActivity
 	public void onStop()
 	{
 		super.onStop();
-
-		if (consoleCmdReceiver != null)
-		{
-			unregisterReceiver(consoleCmdReceiver);
-		}
 	}
 
 	// handle ad popup visibility and requests
