Merging using UE4-To-UE4-Releases-4.10
Add InputDeviceInfo class to proguard-project.txt
#ue4
#android

diff --git a/Engine/Build/Android/Java/proguard-project.txt b/Engine/Build/Android/Java/proguard-project.txt
--- a/Engine/Build/Android/Java/proguard-project.txt
+++ b/Engine/Build/Android/Java/proguard-project.txt
@@ -27,6 +27,9 @@
 }
 
 # Needed...
+-keep class com.epicgames.ue4.GameActivity$InputDeviceInfo {
+   public *;
+}
 -keep class com.epicgames.ue4.GameActivity$JavaAchievement {
    public *;
 }
