Copying up 3 Android resource files somehow missed by Dev-Mobile integration

diff --git a/Engine/Build/Android/Java/res/layout-land/downloader_progress.xml b/Engine/Build/Android/Java/res/layout-land/downloader_progress.xml
new file mode 100644
--- /dev/null
+++ b/Engine/Build/Android/Java/res/layout-land/downloader_progress.xml
@@ -0,0 +1,170 @@
+<?xml version="1.0" encoding="utf-8"?>
+<RelativeLayout xmlns:android="http://schemas.android.com/apk/res/android"
+    android:layout_width="match_parent"
+    android:layout_height="match_parent"
+    android:orientation="horizontal" 
+	android:background="@drawable/downloadimageh"
+	>
+
+    <LinearLayout
+        android:layout_width="match_parent"
+        android:layout_height="wrap_content"
+        android:layout_weight="0"
+        android:orientation="horizontal"
+		android:layout_centerInParent="true"		
+		>
+					
+        <TextView
+            android:id="@+id/statusText"
+            android:layout_width="match_parent"
+            android:layout_height="wrap_content"
+            android:layout_marginBottom="10dp"
+            android:layout_marginLeft="5dp"
+            android:layout_marginTop="10dp"
+            android:textStyle="bold" />
+
+        <LinearLayout
+            android:id="@+id/downloaderDashboard"
+            android:layout_width="fill_parent"
+            android:layout_height="wrap_content"
+            android:layout_below="@id/statusText"
+            android:orientation="horizontal" >
+
+            <RelativeLayout
+                android:layout_width="match_parent"
+                android:layout_height="wrap_content"
+                android:layout_weight="1" >
+
+                <TextView
+                    android:id="@+id/progressAsFraction"
+                    style="@android:style/TextAppearance.Small"
+                    android:layout_width="wrap_content"
+                    android:layout_height="wrap_content"
+                    android:layout_alignParentLeft="true"
+                    android:layout_marginLeft="5dp"
+                    android:text="0MB / 0MB" >
+                </TextView>
+
+                <TextView
+                    android:id="@+id/progressAsPercentage"
+                    style="@android:style/TextAppearance.Small"
+                    android:layout_width="wrap_content"
+                    android:layout_height="wrap_content"
+                    android:layout_alignRight="@+id/progressBar"
+                    android:text="0%" />
+
+                <ProgressBar
+                    android:id="@+id/progressBar"
+                    style="?android:attr/progressBarStyleHorizontal"
+                    android:layout_width="match_parent"
+                    android:layout_height="wrap_content"
+                    android:layout_below="@+id/progressAsFraction"
+                    android:layout_marginBottom="10dp"
+                    android:layout_marginLeft="10dp"
+                    android:layout_marginRight="10dp"
+                    android:layout_marginTop="10dp"
+                    android:layout_weight="1" />
+
+                <TextView
+                    android:id="@+id/progressAverageSpeed"
+                    style="@android:style/TextAppearance.Small"
+                    android:layout_width="wrap_content"
+                    android:layout_height="wrap_content"
+                    android:layout_alignParentLeft="true"
+                    android:layout_below="@+id/progressBar"
+                    android:layout_marginLeft="5dp" />
+
+                <TextView
+                    android:id="@+id/progressTimeRemaining"
+                    style="@android:style/TextAppearance.Small"
+                    android:layout_width="wrap_content"
+                    android:layout_height="wrap_content"
+                    android:layout_alignRight="@+id/progressBar"
+                    android:layout_below="@+id/progressBar" />
+            </RelativeLayout>
+
+            <LinearLayout
+                android:id="@+id/downloaderDashboard"
+                android:layout_width="fill_parent"
+                android:layout_height="wrap_content"
+                android:orientation="horizontal" >
+
+                <Button
+                    android:id="@+id/pauseButton"
+                    android:layout_width="wrap_content"
+                    android:layout_height="wrap_content"
+                    android:layout_gravity="center_vertical"
+                    android:layout_marginBottom="10dp"
+                    android:layout_marginLeft="10dp"
+                    android:layout_marginRight="5dp"
+                    android:layout_marginTop="10dp"
+                    android:layout_weight="0"
+                    android:minHeight="40dp"
+                    android:minWidth="94dp"
+                    android:text="@string/text_button_pause" />
+
+                <Button
+                    android:id="@+id/cancelButton"
+                    android:layout_width="wrap_content"
+                    android:layout_height="wrap_content"
+                    android:layout_gravity="center_vertical"
+                    android:layout_marginBottom="10dp"
+                    android:layout_marginLeft="5dp"
+                    android:layout_marginRight="5dp"
+                    android:layout_marginTop="10dp"
+                    android:layout_weight="0"
+                    android:minHeight="40dp"
+                    android:minWidth="94dp"
+                    android:text="@string/text_button_cancel"
+                    android:visibility="gone" />
+            </LinearLayout>
+        </LinearLayout>
+    </LinearLayout>
+
+    <LinearLayout
+        android:id="@+id/approveCellular"
+        android:layout_width="match_parent"
+        android:layout_height="match_parent"
+        android:layout_weight="1"
+        android:orientation="horizontal"
+        android:visibility="gone" >
+
+        <TextView
+            android:layout_width="match_parent"
+            android:layout_height="wrap_content"
+            android:layout_margin="10dp"
+            android:id="@+id/textPausedParagraph1"
+            android:text="@string/text_paused_cellular" />
+
+        <TextView
+            android:layout_width="match_parent"
+            android:layout_height="wrap_content"
+            android:layout_margin="10dp"
+            android:id="@+id/textPausedParagraph2"
+            android:text="@string/text_paused_cellular_2" />
+
+        <LinearLayout
+            android:id="@+id/buttonRow"
+            android:layout_width="match_parent"
+            android:layout_height="wrap_content"
+            android:orientation="horizontal" >
+
+            <Button
+                android:id="@+id/resumeOverCellular"
+                android:layout_width="wrap_content"
+                android:layout_height="wrap_content"
+                android:layout_gravity="center"
+                android:layout_margin="10dp"
+                android:text="@string/text_button_resume_cellular" />
+
+            <Button
+                android:id="@+id/wifiSettingsButton"
+                android:layout_width="wrap_content"
+                android:layout_height="wrap_content"
+                android:layout_gravity="center"
+                android:layout_margin="10dp"
+                android:text="@string/text_button_wifi_settings" />
+            </LinearLayout>
+    </LinearLayout>
+
+</RelativeLayout>
\ No newline at end of file
diff --git a/Engine/Build/Android/Java/res/layout-port/downloader_progress.xml b/Engine/Build/Android/Java/res/layout-port/downloader_progress.xml
new file mode 100644
--- /dev/null
+++ b/Engine/Build/Android/Java/res/layout-port/downloader_progress.xml
@@ -0,0 +1,170 @@
+<?xml version="1.0" encoding="utf-8"?>
+<RelativeLayout xmlns:android="http://schemas.android.com/apk/res/android"
+    android:layout_width="match_parent"
+    android:layout_height="match_parent"
+    android:orientation="vertical" 
+	android:background="@drawable/downloadimagev"
+	>
+
+    <LinearLayout
+        android:layout_width="match_parent"
+        android:layout_height="wrap_content"
+        android:layout_weight="0"
+        android:orientation="vertical"
+		android:layout_centerInParent="true"		
+		>
+					
+        <TextView
+            android:id="@+id/statusText"
+            android:layout_width="match_parent"
+            android:layout_height="wrap_content"
+            android:layout_marginBottom="10dp"
+            android:layout_marginLeft="5dp"
+            android:layout_marginTop="10dp"
+            android:textStyle="bold" />
+
+        <LinearLayout
+            android:id="@+id/downloaderDashboard"
+            android:layout_width="fill_parent"
+            android:layout_height="wrap_content"
+            android:layout_below="@id/statusText"
+            android:orientation="vertical" >
+
+            <RelativeLayout
+                android:layout_width="match_parent"
+                android:layout_height="wrap_content"
+                android:layout_weight="1" >
+
+                <TextView
+                    android:id="@+id/progressAsFraction"
+                    style="@android:style/TextAppearance.Small"
+                    android:layout_width="wrap_content"
+                    android:layout_height="wrap_content"
+                    android:layout_alignParentLeft="true"
+                    android:layout_marginLeft="5dp"
+                    android:text="0MB / 0MB" >
+                </TextView>
+
+                <TextView
+                    android:id="@+id/progressAsPercentage"
+                    style="@android:style/TextAppearance.Small"
+                    android:layout_width="wrap_content"
+                    android:layout_height="wrap_content"
+                    android:layout_alignRight="@+id/progressBar"
+                    android:text="0%" />
+
+                <ProgressBar
+                    android:id="@+id/progressBar"
+                    style="?android:attr/progressBarStyleHorizontal"
+                    android:layout_width="match_parent"
+                    android:layout_height="wrap_content"
+                    android:layout_below="@+id/progressAsFraction"
+                    android:layout_marginBottom="10dp"
+                    android:layout_marginLeft="10dp"
+                    android:layout_marginRight="10dp"
+                    android:layout_marginTop="10dp"
+                    android:layout_weight="1" />
+
+                <TextView
+                    android:id="@+id/progressAverageSpeed"
+                    style="@android:style/TextAppearance.Small"
+                    android:layout_width="wrap_content"
+                    android:layout_height="wrap_content"
+                    android:layout_alignParentLeft="true"
+                    android:layout_below="@+id/progressBar"
+                    android:layout_marginLeft="5dp" />
+
+                <TextView
+                    android:id="@+id/progressTimeRemaining"
+                    style="@android:style/TextAppearance.Small"
+                    android:layout_width="wrap_content"
+                    android:layout_height="wrap_content"
+                    android:layout_alignRight="@+id/progressBar"
+                    android:layout_below="@+id/progressBar" />
+            </RelativeLayout>
+
+            <LinearLayout
+                android:id="@+id/downloaderDashboard"
+                android:layout_width="fill_parent"
+                android:layout_height="wrap_content"
+                android:orientation="horizontal" >
+
+                <Button
+                    android:id="@+id/pauseButton"
+                    android:layout_width="wrap_content"
+                    android:layout_height="wrap_content"
+                    android:layout_gravity="center_vertical"
+                    android:layout_marginBottom="10dp"
+                    android:layout_marginLeft="10dp"
+                    android:layout_marginRight="5dp"
+                    android:layout_marginTop="10dp"
+                    android:layout_weight="0"
+                    android:minHeight="40dp"
+                    android:minWidth="94dp"
+                    android:text="@string/text_button_pause" />
+
+                <Button
+                    android:id="@+id/cancelButton"
+                    android:layout_width="wrap_content"
+                    android:layout_height="wrap_content"
+                    android:layout_gravity="center_vertical"
+                    android:layout_marginBottom="10dp"
+                    android:layout_marginLeft="5dp"
+                    android:layout_marginRight="5dp"
+                    android:layout_marginTop="10dp"
+                    android:layout_weight="0"
+                    android:minHeight="40dp"
+                    android:minWidth="94dp"
+                    android:text="@string/text_button_cancel"
+                    android:visibility="gone" />
+            </LinearLayout>
+        </LinearLayout>
+    </LinearLayout>
+
+    <LinearLayout
+        android:id="@+id/approveCellular"
+        android:layout_width="match_parent"
+        android:layout_height="match_parent"
+        android:layout_weight="1"
+        android:orientation="vertical"
+        android:visibility="gone" >
+
+        <TextView
+            android:layout_width="match_parent"
+            android:layout_height="wrap_content"
+            android:layout_margin="10dp"
+            android:id="@+id/textPausedParagraph1"
+            android:text="@string/text_paused_cellular" />
+
+        <TextView
+            android:layout_width="match_parent"
+            android:layout_height="wrap_content"
+            android:layout_margin="10dp"
+            android:id="@+id/textPausedParagraph2"
+            android:text="@string/text_paused_cellular_2" />
+
+        <LinearLayout
+            android:id="@+id/buttonRow"
+            android:layout_width="match_parent"
+            android:layout_height="wrap_content"
+            android:orientation="horizontal" >
+
+            <Button
+                android:id="@+id/resumeOverCellular"
+                android:layout_width="wrap_content"
+                android:layout_height="wrap_content"
+                android:layout_gravity="center"
+                android:layout_margin="10dp"
+                android:text="@string/text_button_resume_cellular" />
+
+            <Button
+                android:id="@+id/wifiSettingsButton"
+                android:layout_width="wrap_content"
+                android:layout_height="wrap_content"
+                android:layout_gravity="center"
+                android:layout_margin="10dp"
+                android:text="@string/text_button_wifi_settings" />
+            </LinearLayout>
+    </LinearLayout>
+
+</RelativeLayout>
\ No newline at end of file
