Copying //UE4/GamePlugins-Staging to //UE4/Dev-Main (Source: //GamePlugins/Main @ 3232656)

==========================
MAJOR FEATURES + CHANGES
==========================

description pending [SamZ]

diff --git a/Engine/Build/Android/Java/src/com/epicgames/ue4/GameActivity.java b/Engine/Build/Android/Java/src/com/epicgames/ue4/GameActivity.java
--- a/Engine/Build/Android/Java/src/com/epicgames/ue4/GameActivity.java
+++ b/Engine/Build/Android/Java/src/com/epicgames/ue4/GameActivity.java
@@ -8,6 +8,7 @@ import java.io.File;
 import java.util.Map;
 import java.util.HashMap;
 import java.util.ArrayList;
+import java.text.DecimalFormat;
 
 import android.app.NativeActivity;
 import android.os.Bundle;
@@ -1201,12 +1202,15 @@ public class GameActivity extends NativeActivity implements SurfaceHolder.Callba
 
 		// Set label and starting contents
 		virtualKeyboardAlert.setTitle(Label);
-		virtualKeyboardInputBox.setText("");
-		virtualKeyboardInputBox.append(Contents);
-		virtualKeyboardPreviousContents = Contents;
 
+		// @HSL_BEGIN - Josh.May - 11/01/2016 - Ensure the input mode of the text box is set before setting the contents.
 		// configure for type of input
 		virtualKeyboardInputBox.setInputType(InputType);
+		
+		virtualKeyboardInputBox.setText("");
+		virtualKeyboardInputBox.append(Contents);
+		virtualKeyboardPreviousContents = Contents;
+		// @HSL_END - Josh.May - 11/01/2016
 
 		_activity.runOnUiThread(new Runnable()
 		{
@@ -1651,6 +1655,10 @@ public class GameActivity extends NativeActivity implements SurfaceHolder.Callba
 		{
 			IapStoreHelper = new GooglePlayStoreHelper(InProductKey, this, Log);
 			if (IapStoreHelper != null)
+			{
+				Log.debug("[JAVA] - AndroidThunkJava_IapSetupService - Setup started");
+			}
+			else
 			{
 				Log.debug("[JAVA] - AndroidThunkJava_IapSetupService - Failed to setup IAP service");
 			}
@@ -1664,12 +1672,10 @@ public class GameActivity extends NativeActivity implements SurfaceHolder.Callba
 	
 
 	private String[] CachedQueryProductIDs;
-	private boolean[] CachedQueryConsumables;
-	public boolean AndroidThunkJava_IapQueryInAppPurchases(String[] ProductIDs, boolean[] bConsumable)
+	public boolean AndroidThunkJava_IapQueryInAppPurchases(String[] ProductIDs)
 	{
 		Log.debug("[JAVA] - AndroidThunkJava_IapQueryInAppPurchases");
 		CachedQueryProductIDs = ProductIDs;
-		CachedQueryConsumables = bConsumable;
 
 		boolean bTriggeredQuery = false;
 		if( IapStoreHelper != null )
@@ -1681,7 +1687,7 @@ public class GameActivity extends NativeActivity implements SurfaceHolder.Callba
 				@Override
 				public void run()
 				{
-					IapStoreHelper.QueryInAppPurchases(CachedQueryProductIDs, CachedQueryConsumables);
+					IapStoreHelper.QueryInAppPurchases(CachedQueryProductIDs);
 				}
 			});
 		}
@@ -1776,13 +1782,13 @@ public class GameActivity extends NativeActivity implements SurfaceHolder.Callba
 		}
 	}
 	
-	public boolean AndroidThunkJava_IapBeginPurchase(String ProductId, boolean bConsumable)
+	public boolean AndroidThunkJava_IapBeginPurchase(String ProductId)
 	{
 		Log.debug("[JAVA] - AndroidThunkJava_IapBeginPurchase");
 		boolean bTriggeredPurchase = false;
 		if( IapStoreHelper != null )
 		{
-			bTriggeredPurchase = IapStoreHelper.BeginPurchase(ProductId, bConsumable);
+			bTriggeredPurchase = IapStoreHelper.BeginPurchase(ProductId);
 		}
 		else
 		{
@@ -1806,6 +1812,36 @@ public class GameActivity extends NativeActivity implements SurfaceHolder.Callba
 		return bIsAllowedToMakePurchase;
 	}
 
+	public boolean AndroidThunkJava_IapConsumePurchase(String purchaseToken)
+	{
+		Log.debug("[JAVA] - AndroidThunkJava_IapConsumePurchase " + purchaseToken);
+
+		if( IapStoreHelper != null )
+		{
+			IapStoreHelper.ConsumePurchase(purchaseToken);
+			return true;
+		}
+
+		Log.debug("[JAVA] - Store Helper is invalid");
+		return false;
+	}
+
+	public boolean AndroidThunkJava_IapQueryExistingPurchases()
+	{
+		Log.debug("[JAVA] - AndroidThunkJava_IapQueryExistingPurchases");
+		boolean bTriggeredQuery = false;
+		if( IapStoreHelper != null )
+		{
+			Log.debug("[JAVA] - AndroidThunkJava_IapQueryExistingPurchases - Kick off logic here!");
+			bTriggeredQuery = IapStoreHelper.QueryExistingPurchases();
+		}
+		else
+		{
+			Log.debug("[JAVA] - Store Helper is invalid");
+		}
+		return bTriggeredQuery;
+	}
+
 	public boolean AndroidThunkJava_IapRestorePurchases(String[] InProductIDs, boolean[] bConsumable)
 	{
 		Log.debug("[JAVA] - AndroidThunkJava_IapRestorePurchases");
@@ -2271,6 +2307,20 @@ public class GameActivity extends NativeActivity implements SurfaceHolder.Callba
 
 	public String AndroidThunkJava_GetMetaDataString(String key)
 	{
+		if (key.equals("ue4.displaymetrics.dpi"))
+		{
+			DisplayMetrics metrics = new DisplayMetrics();
+			if (android.os.Build.VERSION.SDK_INT >= 17)
+			{
+				getWindowManager().getDefaultDisplay().getRealMetrics(metrics);
+			} else {
+				// note: not available so get what we can
+				getWindowManager().getDefaultDisplay().getMetrics(metrics);
+			}
+			String screenDPI = new DecimalFormat("###.##").format(metrics.xdpi) + "," + new DecimalFormat("###.##").format(metrics.ydpi) + "," + new DecimalFormat("###.##").format(metrics.densityDpi);
+			return screenDPI;
+		}
+		else
 		if (_bundle == null || key == null)
 		{
 			return null;
diff --git a/Engine/Build/Android/Java/src/com/epicgames/ue4/GooglePlayStoreHelper.java b/Engine/Build/Android/Java/src/com/epicgames/ue4/GooglePlayStoreHelper.java
--- a/Engine/Build/Android/Java/src/com/epicgames/ue4/GooglePlayStoreHelper.java
+++ b/Engine/Build/Android/Java/src/com/epicgames/ue4/GooglePlayStoreHelper.java
@@ -3,6 +3,7 @@
 
 package com.epicgames.ue4;
 
+import android.app.Activity;
 import android.app.PendingIntent;
 import android.content.ComponentName;
 import android.content.Context;
@@ -13,6 +14,7 @@ import android.os.Bundle;
 import android.os.Handler;
 import android.os.IBinder;
 import android.os.Looper;
+import android.os.RemoteException;
 import com.android.vending.billing.util.Base64;
 import com.android.vending.billing.IInAppBillingService;
 import com.android.vending.billing.util.Purchase;
@@ -20,10 +22,39 @@ import java.util.ArrayList;
 import org.json.JSONException;
 import org.json.JSONObject;
 
-
 public class GooglePlayStoreHelper implements StoreHelper
 {
-	private ArrayList<InAppPurchase> InProgressPurchases;
+    // Billing response codes
+	// taken from http://developer.android.com/google/play/billing/billing_reference.html Table 1
+    public static final int BILLING_RESPONSE_RESULT_OK = 0;
+    public static final int BILLING_RESPONSE_RESULT_USER_CANCELED = 1;
+	public static final int BILLING_RESPONSE_RESULT_SERVICE_UNAVAILABLE = 2;
+    public static final int BILLING_RESPONSE_RESULT_BILLING_UNAVAILABLE = 3;
+    public static final int BILLING_RESPONSE_RESULT_ITEM_UNAVAILABLE = 4;
+    public static final int BILLING_RESPONSE_RESULT_DEVELOPER_ERROR = 5;
+    public static final int BILLING_RESPONSE_RESULT_ERROR = 6;
+    public static final int BILLING_RESPONSE_RESULT_ITEM_ALREADY_OWNED = 7;
+    public static final int BILLING_RESPONSE_RESULT_ITEM_NOT_OWNED = 8;
+
+	// Keys for the responses from InAppBillingService
+	// taken from http://developer.android.com/google/play/billing/billing_reference.html Table 3
+	public static final String RESPONSE_CODE = "RESPONSE_CODE";
+	public static final String RESPONSE_INAPP_PURCHASE_DATA = "INAPP_PURCHASE_DATA";
+	public static final String RESPONSE_INAPP_SIGNATURE = "INAPP_DATA_SIGNATURE";
+
+	// taken from https://developer.android.com/google/play/billing/billing_integrate.html#purchase
+	public static final String RESPONSE_GET_SKU_DETAILS_LIST = "DETAILS_LIST";
+	public static final String RESPONSE_BUY_INTENT = "BUY_INTENT";
+	public static final String RESPONSE_INAPP_ITEM_LIST = "INAPP_PURCHASE_ITEM_LIST";
+	public static final String RESPONSE_INAPP_PURCHASE_DATA_LIST = "INAPP_PURCHASE_DATA_LIST";
+	public static final String RESPONSE_INAPP_SIGNATURE_LIST = "INAPP_DATA_SIGNATURE_LIST";
+	public static final String INAPP_CONTINUATION_TOKEN = "INAPP_CONTINUATION_TOKEN";
+
+	// Item types
+	public static final String ITEM_TYPE_INAPP = "inapp";
+
+	// some fields on the getSkuDetails response bundle
+	public static final String GET_SKU_DETAILS_ITEM_LIST = "ITEM_ID_LIST";
 
 	// Our IAB helper interface provided by google.
 	private IInAppBillingService mService;
@@ -61,8 +92,6 @@ public class GooglePlayStoreHelper implements StoreHelper
 		Intent serviceIntent = new Intent("com.android.vending.billing.InAppBillingService.BIND");
 		serviceIntent.setPackage("com.android.vending");
 		gameActivity.bindService(serviceIntent, mServiceConn, Context.BIND_AUTO_CREATE);
-
-		InProgressPurchases = new ArrayList<InAppPurchase>();
 	}
 
 	
@@ -81,7 +110,7 @@ public class GooglePlayStoreHelper implements StoreHelper
 	/**
 	 * Query product details for the provided skus
 	 */
-	public boolean QueryInAppPurchases(String[] InProductIDs, boolean[] bConsumable)
+	public boolean QueryInAppPurchases(String[] InProductIDs)
 	{
 		Log.debug("[GooglePlayStoreHelper] - GooglePlayStoreHelper::QueryInAppPurchases");
 		ArrayList<String> skuList = new ArrayList<String> ();
@@ -93,14 +122,15 @@ public class GooglePlayStoreHelper implements StoreHelper
 		}
 
 		Bundle querySkus = new Bundle();
-		querySkus.putStringArrayList("ITEM_ID_LIST", skuList);
+		querySkus.putStringArrayList(GET_SKU_DETAILS_ITEM_LIST, skuList);
 
 		try
 		{
-			Bundle skuDetails = mService.getSkuDetails(3, gameActivity.getPackageName(), "inapp", querySkus);
+			Bundle skuDetails = mService.getSkuDetails(3, gameActivity.getPackageName(), ITEM_TYPE_INAPP, querySkus);
 
-			int response = skuDetails.getInt("RESPONSE_CODE");
-			if (response == 0)
+			int response = skuDetails.getInt(RESPONSE_CODE);
+			Log.debug("[GooglePlayStoreHelper] - GooglePlayStoreHelper::QueryInAppPurchases - Response " + response + " Bundle:" + skuDetails.toString());
+			if (response == BILLING_RESPONSE_RESULT_OK)
 			{
 				ArrayList<String> productIds = new ArrayList<String>();
 				ArrayList<String> titles = new ArrayList<String>();
@@ -109,7 +139,7 @@ public class GooglePlayStoreHelper implements StoreHelper
 				ArrayList<Float> pricesRaw = new ArrayList<Float>();
 				ArrayList<String> currencyCodes = new ArrayList<String>();
 
-				ArrayList<String> responseList = skuDetails.getStringArrayList("DETAILS_LIST");
+				ArrayList<String> responseList = skuDetails.getStringArrayList(RESPONSE_GET_SKU_DETAILS_LIST);
 				for (String thisResponse : responseList)
 				{
 					JSONObject object = new JSONObject(thisResponse);
@@ -140,11 +170,14 @@ public class GooglePlayStoreHelper implements StoreHelper
 				}
 
 				float[] pricesRawPrimitive = new float[pricesRaw.size()];
-				for(int i = 0; i < pricesRaw.size(); i++)
+				for (int i = 0; i < pricesRaw.size(); i++)
+				{
 					pricesRawPrimitive[i] = pricesRaw.get(i);
+				}
 
-				Log.debug("[GooglePlayStoreHelper] - GooglePlayStoreHelper::QueryInAppPurchases - Success!");
+				Log.debug("[GooglePlayStoreHelper] - GooglePlayStoreHelper::QueryInAppPurchases " + productIds.size() + " items - Success!");
 				nativeQueryComplete(response, productIds.toArray(new String[productIds.size()]), titles.toArray(new String[titles.size()]), descriptions.toArray(new String[descriptions.size()]), prices.toArray(new String[prices.size()]), pricesRawPrimitive, currencyCodes.toArray(new String[currencyCodes.size()]));
+				Log.debug("[GooglePlayStoreHelper] - nativeQueryComplete done!");
 			}
 			else
 			{
@@ -165,55 +198,69 @@ public class GooglePlayStoreHelper implements StoreHelper
 	 * Start the purchase flow for a particular sku
 	 */
 	final int purchaseIntentIdentifier = 1001;
-	public boolean BeginPurchase(String ProductID, boolean bConsumable)
+	public boolean BeginPurchase(String ProductID)
 	{
-		Log.debug("[GooglePlayStoreHelper] - GooglePlayStoreHelper::BeginPurchase");
-		
 		try
 		{
 			String devPayload = GenerateDevPayload(ProductID);
 			Bundle buyIntentBundle = null;
 			int response = -1;
 
-			if (gameActivity.IsInVRMode()) {
+			if (gameActivity.IsInVRMode())
+			{
 				Bundle bundle = new Bundle();
 				bundle.putBoolean("vr", true);
-				response =  mService.isBillingSupportedExtraParams(7, gameActivity.getPackageName(),"inapp", bundle);
-				if (response == 0) {
+				response = mService.isBillingSupportedExtraParams(7, gameActivity.getPackageName(), ITEM_TYPE_INAPP, bundle);
+				if (response == BILLING_RESPONSE_RESULT_OK)
+				{
 					Log.debug("[GooglePlayStoreHelper] - GooglePlayStoreHelper::BeginPurchase - v7 VR purchase" + ProductID);
-					buyIntentBundle = mService.getBuyIntentExtraParams(7, gameActivity.getPackageName(), ProductID, "inapp", devPayload, bundle);
-				} else {
+					buyIntentBundle = mService.getBuyIntentExtraParams(7, gameActivity.getPackageName(), ProductID, ITEM_TYPE_INAPP, devPayload, bundle);
+				} 
+				else
+				{
 					Log.debug("[GooglePlayStoreHelper] - GooglePlayStoreHelper::BeginPurchase - v3 IAB purchase:" + ProductID);
-					buyIntentBundle = mService.getBuyIntent(3, gameActivity.getPackageName(), ProductID, "inapp", devPayload);
+					buyIntentBundle = mService.getBuyIntent(3, gameActivity.getPackageName(), ProductID, ITEM_TYPE_INAPP, devPayload);
 				}
-			} else {
+			}
+			else
+			{
 				Log.debug("[GooglePlayStoreHelper] - GooglePlayStoreHelper::BeginPurchase - v3 IAB purchase:" + ProductID);
-				buyIntentBundle = mService.getBuyIntent(3, gameActivity.getPackageName(), ProductID, "inapp", devPayload);
+				buyIntentBundle = mService.getBuyIntent(3, gameActivity.getPackageName(), ProductID, ITEM_TYPE_INAPP, devPayload);
 			}
-			response = buyIntentBundle.getInt("RESPONSE_CODE");
-			if (response == 0)
+
+			response = buyIntentBundle.getInt(RESPONSE_CODE);
+			if (response == BILLING_RESPONSE_RESULT_OK)
 			{
 				Log.debug("[GooglePlayStoreHelper] - GooglePlayStoreHelper::BeginPurchase - Starting Intent to buy " + ProductID);
-				PendingIntent pendingIntent = buyIntentBundle.getParcelable("BUY_INTENT");
-				PurchaseLaunchCallback callback =  gameActivity.getPurchaseLaunchCallback();
-				if (callback != null) {
-					callback.launchForResult(pendingIntent, purchaseIntentIdentifier);
-				} else {
-					gameActivity.startIntentSenderForResult(pendingIntent.getIntentSender(), purchaseIntentIdentifier, new Intent(), Integer.valueOf(0), Integer.valueOf(0), Integer.valueOf(0));
+				PendingIntent pendingIntent = buyIntentBundle.getParcelable(RESPONSE_BUY_INTENT);
+				if (pendingIntent != null)
+				{
+					PurchaseLaunchCallback callback =  gameActivity.getPurchaseLaunchCallback();
+					if (callback != null)
+					{
+						callback.launchForResult(pendingIntent, purchaseIntentIdentifier);
+					}
+					else
+					{
+						gameActivity.startIntentSenderForResult(pendingIntent.getIntentSender(), purchaseIntentIdentifier, new Intent(), Integer.valueOf(0), Integer.valueOf(0), Integer.valueOf(0));
+					}
+				}
+				else
+				{
+					Log.debug("[GooglePlayStoreHelper] - GooglePlayStoreHelper::BeginPurchase - pendingIntent was null, possible non consumed item");
+					nativePurchaseComplete(UndefinedFailureResponse, ProductID, "", "", "");
 				}
-
-				InProgressPurchases.add(new InAppPurchase(ProductID, bConsumable));
 			}
 			else
 			{
-				Log.debug("[GooglePlayStoreHelper] - GooglePlayStoreHelper::BeginPurchase - Failed with error: " + response);
-				nativePurchaseComplete(response, ProductID, "", "");
+				Log.debug("[GooglePlayStoreHelper] - GooglePlayStoreHelper::BeginPurchase - Failed with error: " + TranslateServerResponseCode(response));
+				nativePurchaseComplete(response, ProductID, "", "", "");
 			}
 		}
 		catch(Exception e)
 		{
 			Log.debug("[GooglePlayStoreHelper] - GooglePlayStoreHelper::BeginPurchase - Failed! " + e.getMessage());
-			nativePurchaseComplete(UndefinedFailureResponse, ProductID, "", "");
+			nativePurchaseComplete(UndefinedFailureResponse, ProductID, "", "", "");
 		}
 
 		return true;
@@ -229,10 +276,10 @@ public class GooglePlayStoreHelper implements StoreHelper
 		ArrayList<String> purchaseDataList = new ArrayList<String>();
 		ArrayList<String> signatureList = new ArrayList<String>();
 		
-		// On first pas the continuation token should be null. 
+		// On first pass the continuation token should be null. 
 		// This will allow us to gather large sets of purchased items recursively
 		int responseCode = GatherOwnedPurchaseData(ownedSkus, purchaseDataList, signatureList, null);
-		if (responseCode == 0)
+		if (responseCode == BILLING_RESPONSE_RESULT_OK)
 		{
 			Log.debug("[GooglePlayStoreHelper] - GooglePlayStoreHelper::RestorePurchases - User has previously purchased " + ownedSkus.size() + " inapp products" );
 
@@ -251,66 +298,80 @@ public class GooglePlayStoreHelper implements StoreHelper
 					try
 					{
 						Log.debug("[GooglePlayStoreHelper] - GooglePlayStoreHelper::RestorePurchases - Now consuming any purchases that may have been missed.");
+						ArrayList<String> productTokens = new ArrayList<String>();
 						ArrayList<String> receipts = new ArrayList<String>();
 						int cachedResponse = 0;
 
+						// @todo possible bug, restores ALL ownedSkus, not just those provided by caller in RestoreProductIDs
 						for (int Idx = 0; Idx < f_ownedSkus.size(); Idx++)
 						{
 							String purchaseData = f_purchaseDataList.get(Idx);
 							String dataSignature = f_signatureList.get(Idx);
-							Purchase purchase = new Purchase("inapp", purchaseData, dataSignature);
-							boolean bTryToConsume = false;
-							int consumeResponse = 0;
-							
-							// This is assuming that all purchases should be consumed. Consuming a purchase that is meant to be a one-time purchase makes it so the
-							// user is able to buy it again. Also, it makes it so the purchase will not be able to be restored again in the future.
 
-							for(int Idy = 0; Idy < RestoreProductIDs.length; Idy++)
+							try
 							{
-								if( purchase.getSku().equals(RestoreProductIDs[Idy]) )
+								Purchase purchase = new Purchase(ITEM_TYPE_INAPP, purchaseData, dataSignature);
+								productTokens.add(purchase.getToken());
+
+								boolean bTryToConsume = false;
+								int consumeResponse = 0;
+							
+								// This is assuming that all purchases should be consumed. Consuming a purchase that is meant to be a one-time purchase makes it so the
+								// user is able to buy it again. Also, it makes it so the purchase will not be able to be restored again in the future.
+
+								for (int Idy = 0; Idy < RestoreProductIDs.length; Idy++)
 								{
-									if(Idy < bRestoreConsumableFlags.length)
+									if( purchase.getSku().equals(RestoreProductIDs[Idy]) )
 									{
-										bTryToConsume = bRestoreConsumableFlags[Idy];
-										Log.debug("[GooglePlayStoreHelper] - GooglePlayStoreHelper::RestorePurchases - Found Consumable Flag for Product " + purchase.getSku() + " bConsumable = " + bTryToConsume);
+										if(Idy < bRestoreConsumableFlags.length)
+										{
+											bTryToConsume = bRestoreConsumableFlags[Idy];
+											Log.debug("[GooglePlayStoreHelper] - GooglePlayStoreHelper::RestorePurchases - Found Consumable Flag for Product " + purchase.getSku() + " bConsumable = " + bTryToConsume);
+										}
+										break;
 									}
-									break;
 								}
-							}
 
-							if (bTryToConsume)
-							{
-								Log.debug("[GooglePlayStoreHelper] - GooglePlayStoreHelper::RestorePurchases - Attempting to consume " + purchase.getSku());
-								consumeResponse = mService.consumePurchase(3, gameActivity.getPackageName(), purchase.getToken());
-							}
+								if (bTryToConsume)
+								{
+									Log.debug("[GooglePlayStoreHelper] - GooglePlayStoreHelper::RestorePurchases - Attempting to consume " + purchase.getSku());
+									consumeResponse = mService.consumePurchase(3, gameActivity.getPackageName(), purchase.getToken());
+								}
 
-							if (consumeResponse == 0)
-							{
-								Log.debug("[GooglePlayStoreHelper] - GooglePlayStoreHelper::RestorePurchases - Purchase restored for " + purchase.getSku());
-								String receipt = Base64.encode(purchase.getOriginalJson().getBytes());
-								receipts.add(receipt);
+								if (consumeResponse == 0)
+								{
+									Log.debug("[GooglePlayStoreHelper] - GooglePlayStoreHelper::RestorePurchases - Purchase restored for " + purchase.getSku());
+									String receipt = Base64.encode(purchase.getOriginalJson().getBytes());
+									receipts.add(receipt);
+								}
+								else
+								{
+									Log.debug("[GooglePlayStoreHelper] - GooglePlayStoreHelper::RestorePurchases - consumePurchase failed for " + purchase.getSku() + " with error " + consumeResponse);
+									receipts.add("");
+									cachedResponse = consumeResponse;
+								}
 							}
-							else
+							catch (JSONException e)
 							{
-								Log.debug("[GooglePlayStoreHelper] - GooglePlayStoreHelper::RestorePurchases - consumePurchase failed for " + purchase.getSku() + " with error " + consumeResponse);
+								Log.debug("[GooglePlayStoreHelper] - GooglePlayStoreHelper::RestorePurchases - Failed to parse receipt! " + e.getMessage());
+								productTokens.add("");
 								receipts.add("");
-								cachedResponse = consumeResponse;
 							}
 						}
 						Log.debug("[GooglePlayStoreHelper] - GooglePlayStoreHelper::RestorePurchases - Success!");
-						nativeRestorePurchasesComplete(cachedResponse, f_ownedSkus.toArray(new String[f_ownedSkus.size()]), receipts.toArray(new String[receipts.size()]), f_signatureList.toArray(new String[f_signatureList.size()]));
+						nativeRestorePurchasesComplete(cachedResponse, f_ownedSkus.toArray(new String[f_ownedSkus.size()]), productTokens.toArray(new String[productTokens.size()]), receipts.toArray(new String[receipts.size()]), f_signatureList.toArray(new String[f_signatureList.size()]));
 					}
 					catch (Exception e)
 					{
 						Log.debug("[GooglePlayStoreHelper] - GooglePlayStoreHelper::RestorePurchases - consumePurchase failed. " + e.getMessage());
-						nativeRestorePurchasesComplete(UndefinedFailureResponse, null, null, null);
+						nativeRestorePurchasesComplete(UndefinedFailureResponse, null, null, null, null);
 					}
 				}
 			});
 		}
 		else
 		{
-			nativeRestorePurchasesComplete(responseCode, null, null, null);
+			nativeRestorePurchasesComplete(responseCode, null, null, null, null);
 			Log.debug("[GooglePlayStoreHelper] - GooglePlayStoreHelper::RestorePurchases - Failed to collect existing purchases");
 			return false;
 		}
@@ -318,7 +379,89 @@ public class GooglePlayStoreHelper implements StoreHelper
 		return true;
 	}
 
-	
+	public boolean QueryExistingPurchases()
+	{
+		Log.debug("[GooglePlayStoreHelper] - GooglePlayStoreHelper::QueryExistingPurchases");
+		ArrayList<String> ownedSkus = new ArrayList<String>();
+		ArrayList<String> purchaseDataList = new ArrayList<String>();
+		ArrayList<String> signatureList = new ArrayList<String>();
+		
+		// On first pass the continuation token should be null. 
+		// This will allow us to gather large sets of purchased items recursively
+		int responseCode = GatherOwnedPurchaseData(ownedSkus, purchaseDataList, signatureList, null);
+		if (responseCode == BILLING_RESPONSE_RESULT_OK)
+		{
+			Log.debug("[GooglePlayStoreHelper] - GooglePlayStoreHelper::QueryExistingPurchases - User has previously purchased " + ownedSkus.size() + " inapp products" );
+
+			ArrayList<String> productTokens = new ArrayList<String>();
+			ArrayList<String> receipts = new ArrayList<String>();
+
+			for (int Idx = 0; Idx < ownedSkus.size(); Idx++)
+			{
+				String purchaseData = purchaseDataList.get(Idx);
+				String dataSignature = signatureList.get(Idx);
+
+				try
+				{
+					Purchase purchase = new Purchase(ITEM_TYPE_INAPP, purchaseData, dataSignature);
+					productTokens.add(purchase.getToken());
+
+					String receipt = Base64.encode(purchase.getOriginalJson().getBytes());
+					receipts.add(receipt);
+				}			
+				catch (JSONException e)
+				{
+					Log.debug("[GooglePlayStoreHelper] - GooglePlayStoreHelper::QueryExistingPurchases - Failed to parse receipt! " + e.getMessage());
+					productTokens.add("");
+					receipts.add("");
+				}
+			}
+
+			Log.debug("[GooglePlayStoreHelper] - GooglePlayStoreHelper::QueryExistingPurchases - Success!");
+			nativeQueryExistingPurchasesComplete(responseCode, ownedSkus.toArray(new String[ownedSkus.size()]), productTokens.toArray(new String[productTokens.size()]), receipts.toArray(new String[receipts.size()]), signatureList.toArray(new String[signatureList.size()]));
+			return true;
+		}
+		else
+		{
+			Log.debug("[GooglePlayStoreHelper] - GooglePlayStoreHelper::QueryExistingPurchases - Failed to collect existing purchases" +  TranslateServerResponseCode(responseCode));
+			nativeQueryExistingPurchasesComplete(responseCode, null, null, null, null);
+			return false;
+		}
+	}
+
+	public void ConsumePurchase(String purchaseToken)
+	{
+		Log.debug("[GooglePlayStoreHelper] - Beginning ConsumePurchase: " + purchaseToken);
+
+		final String f_purchaseToken = purchaseToken;
+
+		Handler mainHandler = new Handler(Looper.getMainLooper());
+		mainHandler.post(new Runnable()
+		{
+			@Override
+			public void run()
+			{
+				try
+				{
+					Log.debug("[GooglePlayStoreHelper] - Consuming token: " + f_purchaseToken);
+					int consumeResponse = mService.consumePurchase(3, gameActivity.getPackageName(), f_purchaseToken);
+					if (consumeResponse == BILLING_RESPONSE_RESULT_OK)
+					{
+						Log.debug("[GooglePlayStoreHelper] - ConsumePurchase success");
+					}
+					else
+					{
+						Log.debug("[GooglePlayStoreHelper] - ConsumePurchase failed with error " + TranslateServerResponseCode(consumeResponse));
+					}
+				}
+				catch(Exception e)
+				{
+					Log.debug("[GooglePlayStoreHelper] - ConsumePurchase failed. " + e.getMessage());
+				}
+			}
+		});
+	}
+
 	///////////////////////////////////////////////////////
 	// IInAppBillingService service registration and unregistration
 
@@ -338,6 +481,27 @@ public class GooglePlayStoreHelper implements StoreHelper
 			Log.debug("[GooglePlayStoreHelper] - ServiceConnection::onServiceConnected");
 			mService = IInAppBillingService.Stub.asInterface(service);
 			bIsIapSetup = true;
+
+			try 
+			{
+                Log.debug("Checking for in-app billing 3 support.");
+
+                // check for in-app billing v3 support
+                int response = mService.isBillingSupported(3, gameActivity.getPackageName(), ITEM_TYPE_INAPP);
+                if (response != BILLING_RESPONSE_RESULT_OK) 
+				{
+					Log.debug("In-app billing version 3 NOT supported for " + gameActivity.getPackageName() + " error " + response);
+                }
+				else
+				{
+					Log.debug("In-app billing version 3 supported for " + gameActivity.getPackageName());
+				}
+			}
+            catch (RemoteException e)
+			{
+				e.printStackTrace();
+				return;
+            }
 		}
 	};
 
@@ -368,91 +532,84 @@ public class GooglePlayStoreHelper implements StoreHelper
 
 		if (requestCode == purchaseIntentIdentifier)
 		{
-			int responseCode = data.getIntExtra("RESPONSE_CODE", 0);
-			Log.debug("[GooglePlayStoreHelper] - GooglePlayStoreHelper::onActivityResult - Processing purchase result. Response Code: " + TranslateServerResponseCode(responseCode));
-			if (responseCode == 0)
+		    if (data == null) 
 			{
-				try
-				{
-					final String purchaseData = data.getStringExtra("INAPP_PURCHASE_DATA");
-					final String dataSignature = data.getStringExtra("INAPP_DATA_SIGNATURE");
-					JSONObject jo = new JSONObject(purchaseData);
-					final String sku = jo.getString("productId");
-					String developerPayload = jo.getString("developerPayload");
-					if (VerifyPayload(developerPayload, sku))
-					{
-						final String purchaseToken = jo.getString("purchaseToken");
-						
-						Handler mainHandler = new Handler(Looper.getMainLooper());
-						mainHandler.post(new Runnable()
-						{
-							@Override
-							public void run()
-							{
-								try
-								{
-									boolean bTryToConsume = true;
-									int consumeResponse = 0;
+				Log.debug("Null data in purchase activity result.");
+				nativePurchaseComplete(UndefinedFailureResponse, "", "", "", "");
+				return true;
+			}
 
-									Log.debug("sku = " + sku + " InProgressPurchases.size() = " + InProgressPurchases.size());
+			int responseCode = data.getIntExtra(RESPONSE_CODE, 0);
+			if (resultCode == Activity.RESULT_OK)
+			{
+				final String purchaseData = data.getStringExtra(RESPONSE_INAPP_PURCHASE_DATA);
+				final String dataSignature = data.getStringExtra(RESPONSE_INAPP_SIGNATURE);
 
-									for(int Idx = 0; Idx < InProgressPurchases.size(); ++Idx)
-									{
-										Log.debug("InProgressPurchases.get(Idx).ProductId = " + InProgressPurchases.get(Idx).ProductId);
-										if( sku.equals( InProgressPurchases.get(Idx).ProductId ) )
-										{
-											bTryToConsume = InProgressPurchases.get(Idx).bConsumable;
-											InProgressPurchases.remove(Idx);
-											Log.debug("Found InProgressPurchase for Product " + sku + " bConsumable = " + bTryToConsume);
-											break;
-										}
-									}
+				Log.debug("[GooglePlayStoreHelper] - GooglePlayStoreHelper::onActivityResult - Processing purchase result. Response Code: " + TranslateServerResponseCode(responseCode));
+				if (responseCode == BILLING_RESPONSE_RESULT_OK)
+				{
+					Log.debug("Purchase data: " + purchaseData);
+					Log.debug("Data signature: " + dataSignature);
+					Log.debug("Extras: " + data.getExtras());
 
-									if(bTryToConsume)
-									{
-										consumeResponse = mService.consumePurchase(3, gameActivity.getPackageName(), purchaseToken);
-									}
+					if (purchaseData == null || dataSignature == null)
+					{
+						Log.debug("[GooglePlayStoreHelper] - GooglePlayStoreHelper::onActivityResult - Either purchaseData or dataSignature is null");
+						nativePurchaseComplete(UndefinedFailureResponse, "", "", "", "");
+					}
 
-									if (consumeResponse == 0)
-									{
-										Purchase purchase = new Purchase("inapp", purchaseData, dataSignature);
-										String receipt = Base64.encode(purchase.getOriginalJson().getBytes());
-										nativePurchaseComplete(consumeResponse, sku, receipt, purchase.getSignature());
-									}
-									else
-									{
-										Log.debug("[GooglePlayStoreHelper] - consumePurchase failed with error " + TranslateServerResponseCode(consumeResponse));
-										nativePurchaseComplete(consumeResponse, sku, "", "");
-									}
-								}
-								catch(Exception e)
+					try
+					{
+						final Purchase purchase = new Purchase(ITEM_TYPE_INAPP, purchaseData, dataSignature);
+						final String sku = purchase.getSku();
+						String developerPayload = purchase.getDeveloperPayload();
+						if (VerifyPayload(developerPayload, sku))
+						{
+							Handler mainHandler = new Handler(Looper.getMainLooper());
+							mainHandler.post(new Runnable()
+							{
+								@Override
+								public void run()
 								{
-									Log.debug("[GooglePlayStoreHelper] - consumePurchase failed. " + e.getMessage());
-									nativePurchaseComplete(UndefinedFailureResponse, sku, "", "");
+									String receipt = Base64.encode(purchase.getOriginalJson().getBytes());
+									nativePurchaseComplete(BILLING_RESPONSE_RESULT_OK, sku, purchase.getToken(), receipt, purchase.getSignature());
 								}
-							}
-						});
+							});
+						}
+						else
+						{
+							Log.debug("[GooglePlayStoreHelper] - GooglePlayStoreHelper::onActivityResult - Failed for verify developer payload for " + sku);
+							nativePurchaseComplete(UndefinedFailureResponse, sku, "", "", "");
+						}
 					}
-					else
+					catch (JSONException e)
 					{
-						Log.debug("[GooglePlayStoreHelper] - GooglePlayStoreHelper::onActivityResult - Failed for verify developer payload for " + sku);
-						nativePurchaseComplete(UndefinedFailureResponse, sku, "", "");
+						Log.debug("[GooglePlayStoreHelper] - GooglePlayStoreHelper::onActivityResult - Failed for purchase request! " + e.getMessage());
+						nativePurchaseComplete(UndefinedFailureResponse, "", "", "", "");
 					}
 				}
-				catch (JSONException e)
+				else
 				{
-					Log.debug("[GooglePlayStoreHelper] - GooglePlayStoreHelper::onActivityResult - Failed for purchase request! " + e.getMessage());
-					nativePurchaseComplete(UndefinedFailureResponse, "", "", "");
+					Log.debug("[GooglePlayStoreHelper] - GooglePlayStoreHelper::onActivityResult - Failed for purchase request!. " + TranslateServerResponseCode(responseCode));
+					nativePurchaseComplete(responseCode, "", "", "", "");
 				}
 			}
+			else if (resultCode == Activity.RESULT_CANCELED)
+			{
+				Log.debug("[GooglePlayStoreHelper] - GooglePlayStoreHelper::onActivityResult - Purchase canceled." + TranslateServerResponseCode(responseCode));
+				nativePurchaseComplete(BILLING_RESPONSE_RESULT_USER_CANCELED, "", "", "", "");
+			}
 			else
 			{
-				Log.debug("[GooglePlayStoreHelper] - GooglePlayStoreHelper::onActivityResult - Failed for purchase request!. " + TranslateServerResponseCode(responseCode));
-				nativePurchaseComplete(responseCode, "", "", "");
+				Log.debug("[GooglePlayStoreHelper] - GooglePlayStoreHelper::onActivityResult - Purchase failed. Result code: " +
+					Integer.toString(resultCode) + ". Response: " + TranslateServerResponseCode(responseCode));
+				nativePurchaseComplete(UndefinedFailureResponse, "", "", "", "");
 			}
+
+			return true;
 		}
 
-		return true;
+		return false;
 	}
 
 	
@@ -464,7 +621,7 @@ public class GooglePlayStoreHelper implements StoreHelper
 	 */
 	private String GenerateDevPayload( String ProductId )
 	{
-		return "ue4."+ProductId;
+		return "ue4." + ProductId;
 	}
 	
 	/**
@@ -487,23 +644,23 @@ public class GooglePlayStoreHelper implements StoreHelper
 		// taken from http://developer.android.com/google/play/billing/billing_reference.html
 		switch(serverResponseCode)
 		{
-			case 0: // BILLING_RESPONSE_RESULT_OK
+			case BILLING_RESPONSE_RESULT_OK:
 				return "Success";
-			case 1: // BILLING_RESPONSE_RESULT_USER_CANCELED
+			case BILLING_RESPONSE_RESULT_USER_CANCELED:
 				return "User pressed back or canceled a dialog";
-			case 2: // BILLING_RESPONSE_RESULT_SERVICE_UNAVAILABLE
+			case BILLING_RESPONSE_RESULT_SERVICE_UNAVAILABLE:
 				return "Network connection is down";
-			case 3: // BILLING_RESPONSE_RESULT_BILLING_UNAVAILABLE
+			case BILLING_RESPONSE_RESULT_BILLING_UNAVAILABLE:
 				return "Billing API version is not supported for the type requested";
-			case 4: // BILLING_RESPONSE_RESULT_ITEM_UNAVAILABLE
+			case BILLING_RESPONSE_RESULT_ITEM_UNAVAILABLE:
 				return "Requested product is not available for purchase";
-			case 5: // BILLING_RESPONSE_RESULT_DEVELOPER_ERROR
+			case BILLING_RESPONSE_RESULT_DEVELOPER_ERROR:
 				return "Invalid arguments provided to the API. This error can also indicate that the application was not correctly signed or properly set up for In-app Billing in Google Play, or does not have the necessary permissions in its manifest";
-			case 6: // BILLING_RESPONSE_RESULT_ERROR
+			case BILLING_RESPONSE_RESULT_ERROR:
 				return "Fatal error during the API action";
-			case 7: // BILLING_RESPONSE_RESULT_ITEM_ALREADY_OWNED
+			case BILLING_RESPONSE_RESULT_ITEM_ALREADY_OWNED:
 				return "Failure to purchase since item is already owned";
-			case 8: // BILLING_RESPONSE_RESULT_ITEM_NOT_OWNED
+			case BILLING_RESPONSE_RESULT_ITEM_NOT_OWNED:
 				return "Failure to consume since item is not owned";
 			default:
 				return "Unknown Server Response Code";
@@ -519,16 +676,16 @@ public class GooglePlayStoreHelper implements StoreHelper
 		int responseCode = UndefinedFailureResponse;
 		try
 		{
-			Bundle ownedItems = mService.getPurchases(3, gameActivity.getPackageName(), "inapp", inContinuationToken);
+			Bundle ownedItems = mService.getPurchases(3, gameActivity.getPackageName(), ITEM_TYPE_INAPP, inContinuationToken);
 
-			responseCode = ownedItems.getInt("RESPONSE_CODE");
-			Log.debug("[GooglePlayStoreHelper] - GooglePlayStoreHelper::RestorePurchases - getPurchases result. Response Code: " + TranslateServerResponseCode(responseCode));
-			if (responseCode == 0)
+			responseCode = ownedItems.getInt(RESPONSE_CODE);
+			Log.debug("[GooglePlayStoreHelper] - GooglePlayStoreHelper::GatherOwnedPurchaseData - getPurchases result. Response Code: " + TranslateServerResponseCode(responseCode));
+			if (responseCode == BILLING_RESPONSE_RESULT_OK)
 			{
-				ArrayList<String> ownedSkus = ownedItems.getStringArrayList("INAPP_PURCHASE_ITEM_LIST");
-				ArrayList<String> purchaseDataList = ownedItems.getStringArrayList("INAPP_PURCHASE_DATA_LIST");
-				ArrayList<String> signatureList = ownedItems.getStringArrayList("INAPP_DATA_SIGNATURE_LIST");
-				String continuationToken = ownedItems.getString("INAPP_CONTINUATION_TOKEN");
+				ArrayList<String> ownedSkus = ownedItems.getStringArrayList(RESPONSE_INAPP_ITEM_LIST);
+				ArrayList<String> purchaseDataList = ownedItems.getStringArrayList(RESPONSE_INAPP_PURCHASE_DATA_LIST);
+				ArrayList<String> signatureList = ownedItems.getStringArrayList(RESPONSE_INAPP_SIGNATURE_LIST);
+				String continuationToken = ownedItems.getString(INAPP_CONTINUATION_TOKEN);
 
 				for (int Idx = 0; Idx < purchaseDataList.size(); Idx++)
 				{
@@ -552,17 +709,16 @@ public class GooglePlayStoreHelper implements StoreHelper
 		}
 		catch (Exception e)
 		{
-			Log.debug("[GooglePlayStoreHelper] - GooglePlayStoreHelper::onActivityResult - Failed for purchase request!. " + e.getMessage());
+			Log.debug("[GooglePlayStoreHelper] - GooglePlayStoreHelper::GatherOwnedPurchaseData - Failed for purchase request!. " + e.getMessage());
 		}
 
 		return responseCode;
 	}
 
-
 	// Callback that notify the C++ implementation that a task has completed
 	public native void nativeQueryComplete(int responseCode, String[] productIDs, String[] titles, String[] descriptions, String[] prices, float[] pricesRaw, String[] currencyCodes );
-	public native void nativePurchaseComplete(int responseCode, String ProductId, String ReceiptData, String Signature);
-	public native void nativeRestorePurchasesComplete(int responseCode, String[] ProductIds, String[] ReceiptsData, String[] Signatures);
-
+	public native void nativePurchaseComplete(int responseCode, String ProductId, String ProductToken, String ReceiptData, String Signature);
+	public native void nativeRestorePurchasesComplete(int responseCode, String[] ProductIds, String[] ProductTokens, String[] ReceiptsData, String[] Signatures);
+	public native void nativeQueryExistingPurchasesComplete(int responseCode, String[] ProductIds, String[] ProductTokens, String[] ReceiptsData, String[] Signatures);
 }
 
diff --git a/Engine/Build/Android/Java/src/com/epicgames/ue4/StoreHelper.java b/Engine/Build/Android/Java/src/com/epicgames/ue4/StoreHelper.java
--- a/Engine/Build/Android/Java/src/com/epicgames/ue4/StoreHelper.java
+++ b/Engine/Build/Android/Java/src/com/epicgames/ue4/StoreHelper.java
@@ -7,21 +7,11 @@ import android.os.Bundle;
 
 public interface StoreHelper
 {
-	public class InAppPurchase
-	{
-		public String ProductId;
-		public boolean bConsumable;
-
-		public InAppPurchase(String InProductId, boolean InConsumable)
-		{
-			ProductId = InProductId;
-			bConsumable = InConsumable;
-		}
-	}
-
-	public boolean QueryInAppPurchases(String[] ProductIDs, boolean[] bConsumable);
-	public boolean BeginPurchase(String ProductID, boolean bConsumable);
+	public boolean QueryInAppPurchases(String[] ProductIDs);
+	public boolean BeginPurchase(String ProductID);
 	public boolean IsAllowedToMakePurchases();
+	public void ConsumePurchase(String purchaseToken);
+	public boolean QueryExistingPurchases();
 	public boolean RestorePurchases(String[] InProductIDs, boolean[] bConsumable);
 	public void	onDestroy();
 	public boolean onActivityResult(int requestCode, int resultCode, Intent data);
diff --git a/Engine/Plugins/Online/Android/OnlineSubsystemGameCircle/Source/GameCircleRuntimeSettings/GameCircleRuntimeSettings.Build.cs b/Engine/Plugins/Online/Android/OnlineSubsystemGameCircle/Source/GameCircleRuntimeSettings/GameCircleRuntimeSettings.Build.cs
--- a/Engine/Plugins/Online/Android/OnlineSubsystemGameCircle/Source/GameCircleRuntimeSettings/GameCircleRuntimeSettings.Build.cs
+++ b/Engine/Plugins/Online/Android/OnlineSubsystemGameCircle/Source/GameCircleRuntimeSettings/GameCircleRuntimeSettings.Build.cs
@@ -7,6 +7,7 @@ public class GameCircleRuntimeSettings : ModuleRules
 	public GameCircleRuntimeSettings(TargetInfo Target)
 	{
 		BinariesSubFolder = "Android";
+		PCHUsage = ModuleRules.PCHUsageMode.UseExplicitOrSharedPCHs;
 
 		PrivateDependencyModuleNames.AddRange(
 			new string[]
diff --git a/Engine/Plugins/Online/Android/OnlineSubsystemGameCircle/Source/OnlineSubsystemGameCircle/OnlineSubsystemGameCircle.Build.cs b/Engine/Plugins/Online/Android/OnlineSubsystemGameCircle/Source/OnlineSubsystemGameCircle/OnlineSubsystemGameCircle.Build.cs
--- a/Engine/Plugins/Online/Android/OnlineSubsystemGameCircle/Source/OnlineSubsystemGameCircle/OnlineSubsystemGameCircle.Build.cs
+++ b/Engine/Plugins/Online/Android/OnlineSubsystemGameCircle/Source/OnlineSubsystemGameCircle/OnlineSubsystemGameCircle.Build.cs
@@ -9,6 +9,7 @@ namespace UnrealBuildTool.Rules
 		public OnlineSubsystemGameCircle(TargetInfo Target)
 		{
 			Definitions.Add("ONLINESUBSYSTEMGAMECIRCLE_PACKAGE=1");
+			PCHUsage = ModuleRules.PCHUsageMode.UseExplicitOrSharedPCHs;
 
 			PrivateIncludePaths.AddRange(
 				new string[] {
diff --git a/Engine/Plugins/Online/Android/OnlineSubsystemGooglePlay/Source/Private/OnlineAsyncTaskGooglePlayAuthAction.cpp b/Engine/Plugins/Online/Android/OnlineSubsystemGooglePlay/Source/Private/OnlineAsyncTaskGooglePlayAuthAction.cpp
--- a/Engine/Plugins/Online/Android/OnlineSubsystemGooglePlay/Source/Private/OnlineAsyncTaskGooglePlayAuthAction.cpp
+++ b/Engine/Plugins/Online/Android/OnlineSubsystemGooglePlay/Source/Private/OnlineAsyncTaskGooglePlayAuthAction.cpp
@@ -2,8 +2,10 @@
 
 #include "OnlineAsyncTaskGooglePlayAuthAction.h"
 
+THIRD_PARTY_INCLUDES_START
 #include "gpg/builder.h"
 #include "gpg/debug.h"
+THIRD_PARTY_INCLUDES_END
 
 FOnlineAsyncTaskGooglePlayAuthAction::FOnlineAsyncTaskGooglePlayAuthAction(
 	FOnlineSubsystemGooglePlay* InSubsystem)
diff --git a/Engine/Plugins/Online/Android/OnlineSubsystemGooglePlay/Source/Private/OnlineAsyncTaskGooglePlayLogin.cpp b/Engine/Plugins/Online/Android/OnlineSubsystemGooglePlay/Source/Private/OnlineAsyncTaskGooglePlayLogin.cpp
--- a/Engine/Plugins/Online/Android/OnlineSubsystemGooglePlay/Source/Private/OnlineAsyncTaskGooglePlayLogin.cpp
+++ b/Engine/Plugins/Online/Android/OnlineSubsystemGooglePlay/Source/Private/OnlineAsyncTaskGooglePlayLogin.cpp
@@ -3,9 +3,11 @@
 #include "OnlineAsyncTaskGooglePlayLogin.h"
 #include "OnlineSubsystemGooglePlay.h"
 
+THIRD_PARTY_INCLUDES_START
 #include "gpg/builder.h"
 #include "gpg/debug.h"
 #include "gpg/types.h"
+THIRD_PARTY_INCLUDES_END
 
 FOnlineAsyncTaskGooglePlayLogin::FOnlineAsyncTaskGooglePlayLogin(
 	FOnlineSubsystemGooglePlay* InSubsystem,
diff --git a/Engine/Plugins/Online/Android/OnlineSubsystemGooglePlay/Source/Private/OnlineAsyncTaskGooglePlayLogout.cpp b/Engine/Plugins/Online/Android/OnlineSubsystemGooglePlay/Source/Private/OnlineAsyncTaskGooglePlayLogout.cpp
--- a/Engine/Plugins/Online/Android/OnlineSubsystemGooglePlay/Source/Private/OnlineAsyncTaskGooglePlayLogout.cpp
+++ b/Engine/Plugins/Online/Android/OnlineSubsystemGooglePlay/Source/Private/OnlineAsyncTaskGooglePlayLogout.cpp
@@ -2,8 +2,11 @@
 
 #include "OnlineAsyncTaskGooglePlayLogout.h"
 #include "OnlineSubsystemGooglePlay.h"
+#include "OnlineIdentityInterfaceGooglePlay.h"
 
+THIRD_PARTY_INCLUDES_START
 #include "gpg/game_services.h"
+THIRD_PARTY_INCLUDES_END
 
 FOnlineAsyncTaskGooglePlayLogout::FOnlineAsyncTaskGooglePlayLogout(
 	FOnlineSubsystemGooglePlay* InSubsystem,
diff --git a/Engine/Plugins/Online/Android/OnlineSubsystemGooglePlay/Source/Private/OnlineAsyncTaskGooglePlayQueryAchievements.cpp b/Engine/Plugins/Online/Android/OnlineSubsystemGooglePlay/Source/Private/OnlineAsyncTaskGooglePlayQueryAchievements.cpp
--- a/Engine/Plugins/Online/Android/OnlineSubsystemGooglePlay/Source/Private/OnlineAsyncTaskGooglePlayQueryAchievements.cpp
+++ b/Engine/Plugins/Online/Android/OnlineSubsystemGooglePlay/Source/Private/OnlineAsyncTaskGooglePlayQueryAchievements.cpp
@@ -2,8 +2,11 @@
 
 #include "OnlineAsyncTaskGooglePlayQueryAchievements.h"
 #include "OnlineSubsystemGooglePlay.h"
+#include "OnlineAchievementsInterfaceGooglePlay.h"
 
+THIRD_PARTY_INCLUDES_START
 #include "gpg/achievement_manager.h"
+THIRD_PARTY_INCLUDES_END
 
 FOnlineAsyncTaskGooglePlayQueryAchievements::FOnlineAsyncTaskGooglePlayQueryAchievements(
 	FOnlineSubsystemGooglePlay* InSubsystem,
diff --git a/Engine/Plugins/Online/Android/OnlineSubsystemGooglePlay/Source/Private/OnlineAsyncTaskGooglePlayQueryInAppPurchases.cpp b/Engine/Plugins/Online/Android/OnlineSubsystemGooglePlay/Source/Private/OnlineAsyncTaskGooglePlayQueryInAppPurchases.cpp
--- a/Engine/Plugins/Online/Android/OnlineSubsystemGooglePlay/Source/Private/OnlineAsyncTaskGooglePlayQueryInAppPurchases.cpp
+++ b/Engine/Plugins/Online/Android/OnlineSubsystemGooglePlay/Source/Private/OnlineAsyncTaskGooglePlayQueryInAppPurchases.cpp
@@ -2,52 +2,96 @@
 
 #include "OnlineAsyncTaskGooglePlayQueryInAppPurchases.h"
 #include "OnlineSubsystemGooglePlay.h"
+#include "OnlineStoreGooglePlayCommon.h"
 
 FOnlineAsyncTaskGooglePlayQueryInAppPurchases::FOnlineAsyncTaskGooglePlayQueryInAppPurchases(
 	FOnlineSubsystemGooglePlay* InSubsystem,
-	const TArray<FString> InProductIds,
-	const TArray<bool> InIsConsumableFlags)
+	const TArray<FString>& InProductIds)
 	: FOnlineAsyncTaskBasic(InSubsystem)
 	, ProductIds(InProductIds)
-	, IsConsumableFlags(InIsConsumableFlags)
 	, bWasRequestSent(false)
 {
 }
 
-
 void FOnlineAsyncTaskGooglePlayQueryInAppPurchases::ProcessQueryAvailablePurchasesResults(bool bInSuccessful)
 {
-	UE_LOG(LogOnline, Display, TEXT("FOnlineAsyncTaskGooglePlayQueryInAppPurchases::ProcessQueryAvailablePurchasesResults"));
+	UE_LOG(LogOnline, Verbose, TEXT("FOnlineAsyncTaskGooglePlayQueryInAppPurchases::ProcessQueryAvailablePurchasesResults"));
 
 	bWasSuccessful = bInSuccessful;
 	bIsComplete = true;
 }
 
-
 void FOnlineAsyncTaskGooglePlayQueryInAppPurchases::Finalize()
 {
-	UE_LOG(LogOnline, Display, TEXT("FOnlineAsyncTaskGooglePlayQueryInAppPurchases::Finalize"));
+	UE_LOG(LogOnline, Verbose, TEXT("FOnlineAsyncTaskGooglePlayQueryInAppPurchases::Finalize"));
 }
 
-
 void FOnlineAsyncTaskGooglePlayQueryInAppPurchases::TriggerDelegates()
 {
-	UE_LOG(LogOnline, Display, TEXT("FOnlineAsyncTaskGooglePlayQueryInAppPurchases::TriggerDelegates"));
+	UE_LOG(LogOnline, Verbose, TEXT("FOnlineAsyncTaskGooglePlayQueryInAppPurchases::TriggerDelegates"));
 	Subsystem->GetStoreInterface()->TriggerOnQueryForAvailablePurchasesCompleteDelegates(bWasSuccessful);
 }
 
-
 void FOnlineAsyncTaskGooglePlayQueryInAppPurchases::Tick()
 {
-//	UE_LOG(LogOnline, Display, TEXT("FOnlineAsyncTaskGooglePlayQueryInAppPurchases::Tick"));
+	UE_LOG(LogOnline, Verbose, TEXT("FOnlineAsyncTaskGooglePlayQueryInAppPurchases::Tick"));
 
 	if (!bWasRequestSent)
 	{
-		extern bool AndroidThunkCpp_Iap_QueryInAppPurchases(const TArray<FString>&, const TArray<bool>&);
-		AndroidThunkCpp_Iap_QueryInAppPurchases(ProductIds, IsConsumableFlags);
+		extern bool AndroidThunkCpp_Iap_QueryInAppPurchases(const TArray<FString>&);
+		AndroidThunkCpp_Iap_QueryInAppPurchases(ProductIds);
 
 		bWasRequestSent = true;
 	}
+}
+
+FOnlineAsyncTaskGooglePlayQueryInAppPurchasesV2::FOnlineAsyncTaskGooglePlayQueryInAppPurchasesV2(
+	FOnlineSubsystemGooglePlay* InSubsystem,
+	const TArray<FUniqueOfferId>& InProductIds,
+	const FOnQueryOnlineStoreOffersComplete& InCompletionDelegate)
+	: FOnlineAsyncTaskBasic(InSubsystem)
+	, ProductIdsV2(InProductIds)
+	, bWasRequestSent(false)
+	, CompletionDelegate(InCompletionDelegate)
+{
+}
+
+void FOnlineAsyncTaskGooglePlayQueryInAppPurchasesV2::ProcessQueryAvailablePurchasesResults(EGooglePlayBillingResponseCode InResponse)
+{
+	UE_LOG(LogOnline, Verbose, TEXT("FOnlineAsyncTaskGooglePlayQueryInAppPurchasesV2::ProcessQueryAvailablePurchasesResults %d"), (int32)InResponse);
+	if (InResponse == EGooglePlayBillingResponseCode::Ok)
+	{
+		bWasSuccessful = true;
+	}
+	else
+	{
+		ErrorStr = ::ToString(InResponse);
+	}
 
+	bIsComplete = true;
+}
+
+void FOnlineAsyncTaskGooglePlayQueryInAppPurchasesV2::Finalize()
+{
+	UE_LOG(LogOnline, Verbose, TEXT("FOnlineAsyncTaskGooglePlayQueryInAppPurchasesV2::Finalize"));
+}
+
+void FOnlineAsyncTaskGooglePlayQueryInAppPurchasesV2::TriggerDelegates()
+{
+	UE_LOG(LogOnline, Verbose, TEXT("FOnlineAsyncTaskGooglePlayQueryInAppPurchasesV2::TriggerDelegates"));
+	CompletionDelegate.ExecuteIfBound(bWasSuccessful, ProductIdsV2, ErrorStr);
+}
+
+void FOnlineAsyncTaskGooglePlayQueryInAppPurchasesV2::Tick()
+{
+	UE_LOG(LogOnline, VeryVerbose, TEXT("FOnlineAsyncTaskGooglePlayQueryInAppPurchasesV2::Tick"));
+
+	if (!bWasRequestSent)
+	{
+		extern bool AndroidThunkCpp_Iap_QueryInAppPurchases(const TArray<FString>&);
+		AndroidThunkCpp_Iap_QueryInAppPurchases(ProductIdsV2);
+
+		bWasRequestSent = true;
+	}
 }
 
diff --git a/Engine/Plugins/Online/Android/OnlineSubsystemGooglePlay/Source/Private/OnlineAsyncTaskGooglePlayQueryInAppPurchases.h b/Engine/Plugins/Online/Android/OnlineSubsystemGooglePlay/Source/Private/OnlineAsyncTaskGooglePlayQueryInAppPurchases.h
--- a/Engine/Plugins/Online/Android/OnlineSubsystemGooglePlay/Source/Private/OnlineAsyncTaskGooglePlayQueryInAppPurchases.h
+++ b/Engine/Plugins/Online/Android/OnlineSubsystemGooglePlay/Source/Private/OnlineAsyncTaskGooglePlayQueryInAppPurchases.h
@@ -4,16 +4,21 @@
 
 #include "CoreMinimal.h"
 #include "OnlineAsyncTaskManager.h"
+#include "OnlineStoreInterfaceV2.h"
+#include "OnlineStoreGooglePlayCommon.h"
 
 class FOnlineSubsystemGooglePlay;
 
+/**
+ * Async event tracking an outstanding "query available in app purchases" task for StoreV1
+ */
 class FOnlineAsyncTaskGooglePlayQueryInAppPurchases : public FOnlineAsyncTaskBasic<FOnlineSubsystemGooglePlay>
 {
+
 public:
 	FOnlineAsyncTaskGooglePlayQueryInAppPurchases(
 		FOnlineSubsystemGooglePlay* InSubsystem,
-		const TArray<FString> InProductIds,
-		const TArray<bool> InIsConsumableFlags);
+		const TArray<FString>& InProductIds);
 
 	// FOnlineAsyncItem
 	virtual FString ToString() const override { return TEXT("QueryInAppPurchases"); }
@@ -26,14 +31,51 @@ public:
 	// FOnlineAsyncTaskGooglePlayQueryInAppPurchases
 	void ProcessQueryAvailablePurchasesResults(bool bInSuccessful);
 
-private:
+protected:
 
 	// The product ids provided for this task
 	const TArray<FString> ProductIds;
 
-	// The consume flags provided for this task
-	const TArray<bool> IsConsumableFlags;
+	/** Flag indicating that the request has been sent */
+	bool bWasRequestSent;
+};
+
+/**
+* Async event tracking an outstanding "query available in app purchases" task for StoreV2
+*/
+class FOnlineAsyncTaskGooglePlayQueryInAppPurchasesV2 : public FOnlineAsyncTaskBasic<FOnlineSubsystemGooglePlay>
+{
+
+public:
+	FOnlineAsyncTaskGooglePlayQueryInAppPurchasesV2(
+		FOnlineSubsystemGooglePlay* InSubsystem,
+		const TArray<FUniqueOfferId>& InProductIds,
+		const FOnQueryOnlineStoreOffersComplete& InCompletionDelegate);
+
+	// FOnlineAsyncItem
+	virtual FString ToString() const override { return TEXT("QueryInAppPurchasesV2"); }
+	virtual void Finalize() override;
+	virtual void TriggerDelegates() override;
+
+	// FOnlineAsyncTask
+	virtual void Tick() override;
+
+	/**
+	 * Available in app purchase query is complete, finish the task before notifying external listeners
+	 */
+	void ProcessQueryAvailablePurchasesResults(EGooglePlayBillingResponseCode InResponse);
+
+protected:
+
+	/** The product ids provided for this task */
+	const TArray<FUniqueOfferId> ProductIdsV2;
+
+	/** Completion delegate passed in at the time of the purchase query */
+	const FOnQueryOnlineStoreOffersComplete CompletionDelegate;
 
 	/** Flag indicating that the request has been sent */
 	bool bWasRequestSent;
+
+	/** String storing any generated errors during the query */
+	FString ErrorStr;
 };
diff --git a/Engine/Plugins/Online/Android/OnlineSubsystemGooglePlay/Source/Private/OnlineAsyncTaskGooglePlayReadLeaderboard.cpp b/Engine/Plugins/Online/Android/OnlineSubsystemGooglePlay/Source/Private/OnlineAsyncTaskGooglePlayReadLeaderboard.cpp
--- a/Engine/Plugins/Online/Android/OnlineSubsystemGooglePlay/Source/Private/OnlineAsyncTaskGooglePlayReadLeaderboard.cpp
+++ b/Engine/Plugins/Online/Android/OnlineSubsystemGooglePlay/Source/Private/OnlineAsyncTaskGooglePlayReadLeaderboard.cpp
@@ -2,8 +2,12 @@
 
 #include "OnlineAsyncTaskGooglePlayReadLeaderboard.h"
 #include "OnlineSubsystemGooglePlay.h"
+#include "OnlineIdentityInterfaceGooglePlay.h"
+#include "OnlineLeaderboardInterfaceGooglePlay.h"
 
+THIRD_PARTY_INCLUDES_START
 #include "gpg/leaderboard_manager.h"
+THIRD_PARTY_INCLUDES_END
 
 using namespace gpg;
 
diff --git a/Engine/Plugins/Online/Android/OnlineSubsystemGooglePlay/Source/Private/OnlineAsyncTaskGooglePlayShowLoginUI.cpp b/Engine/Plugins/Online/Android/OnlineSubsystemGooglePlay/Source/Private/OnlineAsyncTaskGooglePlayShowLoginUI.cpp
--- a/Engine/Plugins/Online/Android/OnlineSubsystemGooglePlay/Source/Private/OnlineAsyncTaskGooglePlayShowLoginUI.cpp
+++ b/Engine/Plugins/Online/Android/OnlineSubsystemGooglePlay/Source/Private/OnlineAsyncTaskGooglePlayShowLoginUI.cpp
@@ -2,9 +2,11 @@
 
 #include "OnlineAsyncTaskGooglePlayShowLoginUI.h"
 #include "OnlineSubsystemGooglePlay.h"
+#include "OnlineIdentityInterfaceGooglePlay.h"
 
+THIRD_PARTY_INCLUDES_START
 #include "gpg/player_manager.h"
-
+THIRD_PARTY_INCLUDES_END
 
 FOnlineAsyncTaskGooglePlayShowLoginUI::FOnlineAsyncTaskGooglePlayShowLoginUI(
 	FOnlineSubsystemGooglePlay* InSubsystem,
diff --git a/Engine/Plugins/Online/Android/OnlineSubsystemGooglePlay/Source/Private/OnlineAsyncTaskGooglePlayShowLoginUI.h b/Engine/Plugins/Online/Android/OnlineSubsystemGooglePlay/Source/Private/OnlineAsyncTaskGooglePlayShowLoginUI.h
--- a/Engine/Plugins/Online/Android/OnlineSubsystemGooglePlay/Source/Private/OnlineAsyncTaskGooglePlayShowLoginUI.h
+++ b/Engine/Plugins/Online/Android/OnlineSubsystemGooglePlay/Source/Private/OnlineAsyncTaskGooglePlayShowLoginUI.h
@@ -9,8 +9,11 @@
 #include "OnlineSubsystemGooglePlayPackage.h"
 
 #include <vector>
+
+THIRD_PARTY_INCLUDES_START
 #include "gpg/types.h"
 #include "gpg/player_manager.h"
+THIRD_PARTY_INCLUDES_END
 
 
 class FOnlineSubsystemGooglePlay;
diff --git a/Engine/Plugins/Online/Android/OnlineSubsystemGooglePlay/Source/Private/OnlineExternalUIInterfaceGooglePlay.cpp b/Engine/Plugins/Online/Android/OnlineSubsystemGooglePlay/Source/Private/OnlineExternalUIInterfaceGooglePlay.cpp
--- a/Engine/Plugins/Online/Android/OnlineSubsystemGooglePlay/Source/Private/OnlineExternalUIInterfaceGooglePlay.cpp
+++ b/Engine/Plugins/Online/Android/OnlineSubsystemGooglePlay/Source/Private/OnlineExternalUIInterfaceGooglePlay.cpp
@@ -5,8 +5,10 @@
 #include "UObject/Class.h"
 #include "OnlineSubsystemGooglePlay.h"
 
+THIRD_PARTY_INCLUDES_START
 #include "gpg/achievement_manager.h"
 #include "gpg/leaderboard_manager.h"
+THIRD_PARTY_INCLUDES_END
 
 FOnlineExternalUIGooglePlay::FOnlineExternalUIGooglePlay(FOnlineSubsystemGooglePlay* InSubsystem)
 	: Subsystem(InSubsystem)
@@ -25,7 +27,7 @@ bool FOnlineExternalUIGooglePlay::ShowFriendsUI(int32 LocalUserNum)
 	return false;
 }
 
-bool FOnlineExternalUIGooglePlay::ShowInviteUI(int32 LocalUserNum, FName SessionMame)
+bool FOnlineExternalUIGooglePlay::ShowInviteUI(int32 LocalUserNum, FName SessionName)
 {
 	return false;
 }
diff --git a/Engine/Plugins/Online/Android/OnlineSubsystemGooglePlay/Source/Private/OnlineExternalUIInterfaceGooglePlay.h b/Engine/Plugins/Online/Android/OnlineSubsystemGooglePlay/Source/Private/OnlineExternalUIInterfaceGooglePlay.h
--- a/Engine/Plugins/Online/Android/OnlineSubsystemGooglePlay/Source/Private/OnlineExternalUIInterfaceGooglePlay.h
+++ b/Engine/Plugins/Online/Android/OnlineSubsystemGooglePlay/Source/Private/OnlineExternalUIInterfaceGooglePlay.h
@@ -24,7 +24,7 @@ public:
 
 	virtual bool ShowLoginUI(const int ControllerIndex, bool bShowOnlineOnly, const FOnLoginUIClosedDelegate& Delegate = FOnLoginUIClosedDelegate()) override;
 	virtual bool ShowFriendsUI(int32 LocalUserNum) override;
-	virtual bool ShowInviteUI(int32 LocalUserNum, FName SessionMame) override;
+	virtual bool ShowInviteUI(int32 LocalUserNum, FName SessionName) override;
 	virtual bool ShowAchievementsUI(int32 LocalUserNum) override;
 	virtual bool ShowLeaderboardUI(const FString& LeaderboardName) override;
 	virtual bool ShowWebURL(const FString& Url, const FShowWebUrlParams& ShowParams, const FOnShowWebUrlClosedDelegate& Delegate = FOnShowWebUrlClosedDelegate()) override;
diff --git a/Engine/Plugins/Online/Android/OnlineSubsystemGooglePlay/Source/Private/OnlineIdentityInterfaceGooglePlay.cpp b/Engine/Plugins/Online/Android/OnlineSubsystemGooglePlay/Source/Private/OnlineIdentityInterfaceGooglePlay.cpp
--- a/Engine/Plugins/Online/Android/OnlineSubsystemGooglePlay/Source/Private/OnlineIdentityInterfaceGooglePlay.cpp
+++ b/Engine/Plugins/Online/Android/OnlineSubsystemGooglePlay/Source/Private/OnlineIdentityInterfaceGooglePlay.cpp
@@ -13,9 +13,8 @@ FOnlineIdentityGooglePlay::FOnlineIdentityGooglePlay(FOnlineSubsystemGooglePlay*
 	, MainSubsystem(InSubsystem)
 	, bRegisteringUser(false)
 	, bLoggingInUser(false)
-
 {
-	UE_LOG(LogOnline, Display, TEXT("FOnlineIdentityAndroid::FOnlineIdentityAndroid()"));
+	UE_LOG(LogOnline, Display, TEXT("FOnlineIdentityGooglePlay::FOnlineIdentityGooglePlay()"));
 	check(MainSubsystem != nullptr);
 	PendingConnectRequest.ConnectionInterface = this;
 }
@@ -33,7 +32,6 @@ TArray<TSharedPtr<FUserOnlineAccount> > FOnlineIdentityGooglePlay::GetAllUserAcc
 	return Result;
 }
 
-
 bool FOnlineIdentityGooglePlay::Login(int32 LocalUserNum, const FOnlineAccountCredentials& AccountCredentials) 
 {
 	bool bStartedLogin = false;
@@ -136,19 +134,19 @@ TSharedPtr<const FUniqueNetId> FOnlineIdentityGooglePlay::CreateUniquePlayerId(c
 
 FString FOnlineIdentityGooglePlay::GetPlayerNickname(int32 LocalUserNum) const
 {
-	UE_LOG(LogOnline, Display, TEXT("FOnlineIdentityAndroid::GetPlayerNickname"));
+	UE_LOG(LogOnline, Display, TEXT("FOnlineIdentityGooglePlay::GetPlayerNickname"));
 	return PlayerAlias;
 }
 
 FString FOnlineIdentityGooglePlay::GetPlayerNickname(const FUniqueNetId& UserId) const
 {
-	UE_LOG(LogOnline, Display, TEXT("FOnlineIdentityAndroid::GetPlayerNickname"));
+	UE_LOG(LogOnline, Display, TEXT("FOnlineIdentityGooglePlay::GetPlayerNickname"));
 	return PlayerAlias;
 }
 
 FString FOnlineIdentityGooglePlay::GetAuthToken(int32 LocalUserNum) const
 {
-	UE_LOG(LogOnline, Display, TEXT("FOnlineIdentityAndroid::GetAuthToken"));
+	UE_LOG(LogOnline, Display, TEXT("FOnlineIdentityGooglePlay::GetAuthToken"));
 	return AuthToken;
 }
 
diff --git a/Engine/Plugins/Online/Android/OnlineSubsystemGooglePlay/Source/Private/OnlineLeaderboardInterfaceGooglePlay.cpp b/Engine/Plugins/Online/Android/OnlineSubsystemGooglePlay/Source/Private/OnlineLeaderboardInterfaceGooglePlay.cpp
--- a/Engine/Plugins/Online/Android/OnlineSubsystemGooglePlay/Source/Private/OnlineLeaderboardInterfaceGooglePlay.cpp
+++ b/Engine/Plugins/Online/Android/OnlineSubsystemGooglePlay/Source/Private/OnlineLeaderboardInterfaceGooglePlay.cpp
@@ -6,7 +6,9 @@
 #include "UObject/Class.h"
 #include "OnlineSubsystemGooglePlay.h"
 
+THIRD_PARTY_INCLUDES_START
 #include "gpg/leaderboard_manager.h"
+THIRD_PARTY_INCLUDES_END
 
 FOnlineLeaderboardsGooglePlay::FOnlineLeaderboardsGooglePlay(FOnlineSubsystemGooglePlay* InSubsystem)
 	: Subsystem(InSubsystem)
@@ -78,7 +80,7 @@ bool FOnlineLeaderboardsGooglePlay::WriteLeaderboards(const FName& SessionName,
 				FOnlinePendingLeaderboardWrite* UnreportedScore = new (UnreportedScores) FOnlinePendingLeaderboardWrite();
 				UnreportedScore->LeaderboardName = LeaderboardName;
 				UnreportedScore->Score = Score;
-				UE_LOG_ONLINE(Display, TEXT("FOnlineLeaderboardsAndroid::WriteLeaderboards() Int64 value Score: %d"), Score);
+				UE_LOG_ONLINE(Display, TEXT("FOnlineLeaderboardsGooglePlay::WriteLeaderboards() Int64 value Score: %d"), Score);
 
 				bWroteAnyLeaderboard = true;
 			}
@@ -92,7 +94,7 @@ bool FOnlineLeaderboardsGooglePlay::WriteLeaderboards(const FName& SessionName,
 				FOnlinePendingLeaderboardWrite* UnreportedScore = new (UnreportedScores) FOnlinePendingLeaderboardWrite();
 				UnreportedScore->LeaderboardName = LeaderboardName;
 				UnreportedScore->Score = Score;
-				UE_LOG_ONLINE(Display, TEXT("FOnlineLeaderboardsAndroid::WriteLeaderboards() Int32 value Score: %d "), Score);
+				UE_LOG_ONLINE(Display, TEXT("FOnlineLeaderboardsGooglePlay::WriteLeaderboards() Int32 value Score: %d "), Score);
 
 				bWroteAnyLeaderboard = true;
 			}
diff --git a/Engine/Plugins/Online/Android/OnlineSubsystemGooglePlay/Source/Private/OnlinePurchaseGooglePlay.cpp b/Engine/Plugins/Online/Android/OnlineSubsystemGooglePlay/Source/Private/OnlinePurchaseGooglePlay.cpp
new file mode 100644
--- /dev/null
+++ b/Engine/Plugins/Online/Android/OnlineSubsystemGooglePlay/Source/Private/OnlinePurchaseGooglePlay.cpp
@@ -0,0 +1,392 @@
+// Copyright 1998-2016 Epic Games, Inc. All Rights Reserved.
+
+#include "OnlinePurchaseGooglePlay.h"
+#include "OnlineSubsystemGooglePlay.h"
+#include "OnlineStoreGooglePlay.h"
+#include "OnlineError.h"
+#include "Base64.h"
+
+#define LOCTEXT_NAMESPACE "OnlineSubsystemGooglePlay"
+#define GOOGLEPLAYUSER TEXT("GooglePlayUser")
+
+FGoogleTransactionData::FGoogleTransactionData(const FString& InOfferId, const FString& InProductToken, const FString& InReceiptData, const FString& InSignature)
+	: OfferId(InOfferId)
+	, TransactionIdentifier(InProductToken)
+	, CombinedTransactionData(InReceiptData, InSignature)
+{
+	if (TransactionIdentifier.IsEmpty())
+	{
+		ErrorStr = TEXT("Receipt does not contain purchase token");
+	}
+	else if (CombinedTransactionData.ReceiptData.IsEmpty())
+	{
+		ErrorStr = TEXT("Receipt does not contain receipt data");
+	}
+	else if (CombinedTransactionData.Signature.IsEmpty())
+	{
+		ErrorStr = TEXT("Receipt does not contain signature data");
+	}
+}
+
+FOnlinePurchaseGooglePlay::FOnlinePurchaseGooglePlay(FOnlineSubsystemGooglePlay* InSubsystem)
+	: bQueryingReceipts(false)
+	, Subsystem(InSubsystem)
+{
+	UE_LOG(LogOnline, Verbose, TEXT( "FOnlinePurchaseGooglePlay::FOnlinePurchaseGooglePlay" ));
+}
+
+FOnlinePurchaseGooglePlay::FOnlinePurchaseGooglePlay()
+{
+	UE_LOG(LogOnline, Verbose, TEXT( "FOnlinePurchaseGooglePlay::FOnlinePurchaseGooglePlay" ));
+}
+
+FOnlinePurchaseGooglePlay::~FOnlinePurchaseGooglePlay()
+{
+	if (Subsystem)
+	{
+		Subsystem->ClearOnGooglePlayProcessPurchaseCompleteDelegate_Handle(ProcessPurchaseResultDelegateHandle);
+		Subsystem->ClearOnGooglePlayQueryExistingPurchasesCompleteDelegate_Handle(QueryExistingPurchasesCompleteDelegateHandle);
+	}
+}
+
+void FOnlinePurchaseGooglePlay::Init()
+{
+	UE_LOG(LogOnline, Verbose, TEXT("FOnlinePurchaseGooglePlay::Init"));
+
+	FOnGooglePlayProcessPurchaseCompleteDelegate PurchaseCompleteDelegate = FOnGooglePlayProcessPurchaseCompleteDelegate::CreateThreadSafeSP(this, &FOnlinePurchaseGooglePlay::OnTransactionCompleteResponse);
+	ProcessPurchaseResultDelegateHandle = Subsystem->AddOnGooglePlayProcessPurchaseCompleteDelegate_Handle(PurchaseCompleteDelegate);
+
+	FOnGooglePlayQueryExistingPurchasesCompleteDelegate QueryExistingPurchasesCompleteDelegate = FOnGooglePlayQueryExistingPurchasesCompleteDelegate::CreateThreadSafeSP(this, &FOnlinePurchaseGooglePlay::OnQueryExistingPurchasesComplete);
+	QueryExistingPurchasesCompleteDelegateHandle = Subsystem->AddOnGooglePlayQueryExistingPurchasesCompleteDelegate_Handle(QueryExistingPurchasesCompleteDelegate);
+}
+
+bool FOnlinePurchaseGooglePlay::IsAllowedToPurchase(const FUniqueNetId& UserId)
+{
+	UE_LOG(LogOnline, Verbose, TEXT("FOnlinePurchaseGooglePlay::IsAllowedToPurchase"));
+	extern bool AndroidThunkCpp_Iap_IsAllowedToMakePurchases();
+	return AndroidThunkCpp_Iap_IsAllowedToMakePurchases();
+}
+
+void FOnlinePurchaseGooglePlay::Checkout(const FUniqueNetId& UserId, const FPurchaseCheckoutRequest& CheckoutRequest, const FOnPurchaseCheckoutComplete& Delegate)
+{
+	bool bStarted = false;
+	FText ErrorMessage;
+
+	TSharedRef<FOnlinePurchasePendingTransactionGooglePlay> RequestedTransaction = MakeShareable(new FOnlinePurchasePendingTransactionGooglePlay(CheckoutRequest, UserId, EPurchaseTransactionState::NotStarted, Delegate));
+
+	if (IsAllowedToPurchase(UserId))
+	{
+		const FString UserIdStr = GOOGLEPLAYUSER;
+		const TSharedRef<FOnlinePurchasePendingTransactionGooglePlay>* UserPendingTransaction = PendingTransactions.Find(UserIdStr);
+		if (UserPendingTransaction == nullptr)
+		{
+			FOnlineStoreGooglePlayV2Ptr StoreInterface = StaticCastSharedPtr<FOnlineStoreGooglePlayV2>(Subsystem->GetStoreV2Interface());
+			if (StoreInterface.IsValid())
+			{
+				TSharedRef<FOnlinePurchasePendingTransactionGooglePlay> PendingTransaction = PendingTransactions.Add(UserIdStr, RequestedTransaction);
+
+				int32 NumOffers = CheckoutRequest.PurchaseOffers.Num();
+				if (NumOffers > 0)
+				{
+					const FPurchaseCheckoutRequest::FPurchaseOfferEntry& Offer = CheckoutRequest.PurchaseOffers[0];
+
+					extern bool AndroidThunkCpp_Iap_BeginPurchase(const FString&);
+					bStarted = AndroidThunkCpp_Iap_BeginPurchase(Offer.OfferId);
+					UE_LOG(LogOnline, Display, TEXT("Created Transaction? - %s"),
+						bStarted ? TEXT("Created a transaction.") : TEXT("Failed to create a transaction."));
+
+					RequestedTransaction->PendingPurchaseInfo.TransactionState = bStarted ? EPurchaseTransactionState::Processing : EPurchaseTransactionState::Failed;
+					if (NumOffers > 1)
+					{
+						UE_LOG(LogOnline, Warning, TEXT("GooglePlay supports purchasing one offer at a time, %d were requested and ignored"), NumOffers - 1);
+					}
+				}
+				else
+				{
+					ErrorMessage = NSLOCTEXT("GooglePlayPurchase", "ErrorNoOffersSpecified", "Failed to checkout, no offers given.");
+					RequestedTransaction->PendingPurchaseInfo.TransactionState = EPurchaseTransactionState::Failed;
+				}
+			}
+		}
+		else
+		{
+			ErrorMessage = NSLOCTEXT("GooglePlayPurchase", "ErrorTransactionInProgress", "Failed to checkout, user has in progress transaction.");
+			RequestedTransaction->PendingPurchaseInfo.TransactionState = EPurchaseTransactionState::Failed;
+		}
+	}
+	else
+	{
+		ErrorMessage = NSLOCTEXT("GooglePlayPurchase", "ErrorPurchaseNotAllowed", "Failed to checkout, user not allowed to purchase.");
+		RequestedTransaction->PendingPurchaseInfo.TransactionState = EPurchaseTransactionState::Failed;
+	}
+
+	if (!bStarted)
+	{
+		TSharedRef<FPurchaseReceipt> FailReceipt = RequestedTransaction->GenerateReceipt();
+
+		Subsystem->ExecuteNextTick([ErrorMessage, FailReceipt, Delegate]()
+		{
+			FOnlineError Error(ErrorMessage);
+			Delegate.ExecuteIfBound(Error, FailReceipt);
+		});
+	}
+}
+
+void FOnlinePurchaseGooglePlay::FinalizePurchase(const FUniqueNetId& UserId, const FString& ReceiptId)
+{
+	UE_LOG(LogOnline, Verbose, TEXT("FOnlinePurchaseGooglePlay::FinalizePurchase %s %s"), *UserId.ToString(), *ReceiptId);
+	extern void AndroidThunkCpp_Iap_ConsumePurchase(const FString&);
+	AndroidThunkCpp_Iap_ConsumePurchase(ReceiptId);
+}
+
+void FOnlinePurchaseGooglePlay::RedeemCode(const FUniqueNetId& UserId, const FRedeemCodeRequest& RedeemCodeRequest, const FOnPurchaseRedeemCodeComplete& Delegate)
+{
+	// Not supported
+	FOnlineError Result;
+	Delegate.ExecuteIfBound(Result, MakeShareable(new FPurchaseReceipt()));
+}
+
+void FOnlinePurchaseGooglePlay::QueryReceipts(const FUniqueNetId& UserId, bool bRestoreReceipts, const FOnQueryReceiptsComplete& Delegate)
+{
+	bool bSuccess = false;
+	bool bTriggerDelegate = true;
+	if (!bQueryingReceipts)
+	{
+		/**
+		 * bRestoreReceipts is irrelevant because GooglePlay requires the client to consume purchase in order to make a new purchase
+		 * There is no concept of restore here, any purchase query will reveal any non consumed purchases. They will remain in the list until the the game
+		 * consumes it via FinalizePurchase().
+		 */
+		bQueryingReceipts = true;
+		QueryReceiptsComplete = Delegate;
+
+		extern bool AndroidThunkCpp_Iap_QueryExistingPurchases();
+		if (AndroidThunkCpp_Iap_QueryExistingPurchases())
+		{
+			bTriggerDelegate = false;
+			bSuccess = true;
+		}
+		else
+		{
+			UE_LOG(LogOnline, Verbose, TEXT("FOnlinePurchaseGooglePlay::QueryReceipts failed to start query"));
+		}
+	}
+	else
+	{
+		UE_LOG(LogOnline, Verbose, TEXT("FOnlinePurchaseGooglePlay::QueryReceipts already in progress."));
+	}
+
+	if (bTriggerDelegate)
+	{
+		Subsystem->ExecuteNextTick([this, Delegate, bSuccess]() {
+			FOnlineError Result(bSuccess);
+			Delegate.ExecuteIfBound(Result);
+			bQueryingReceipts = false;
+			QueryReceiptsComplete.Unbind();
+		});
+	}
+}
+
+void FOnlinePurchaseGooglePlay::GetReceipts(const FUniqueNetId& UserId, TArray<FPurchaseReceipt>& OutReceipts) const
+{
+	OutReceipts.Empty();
+
+	// Add the cached list of user purchases
+	FString UserIdStr = GOOGLEPLAYUSER;
+	const TArray<TSharedRef<FPurchaseReceipt>>* UserCompletedTransactionsPtr = CompletedTransactions.Find(UserIdStr);
+	if (UserCompletedTransactionsPtr != nullptr)
+	{
+		const TArray<TSharedRef<FPurchaseReceipt>>& UserCompletedTransactions = *UserCompletedTransactionsPtr;
+		for (int32 Idx = 0; Idx < UserCompletedTransactions.Num(); Idx++)
+		{
+			const TSharedRef<FPurchaseReceipt>& Transaction = UserCompletedTransactions[Idx];
+			OutReceipts.Add(*Transaction);
+		}
+	}
+	
+	// Add purchases completed while "offline"
+	for (int32 Idx = 0; Idx < OfflineTransactions.Num(); Idx++)
+	{
+		OutReceipts.Add(*OfflineTransactions[Idx]);
+	}
+}
+
+void FOnlinePurchaseGooglePlay::OnTransactionCompleteResponse(EGooglePlayBillingResponseCode InResponseCode, const FGoogleTransactionData& InTransactionData)
+{
+	UE_LOG(LogOnline, Verbose, TEXT("FOnlinePurchaseGooglePlay::OnTransactionCompleteResponse %s"), ToString(InResponseCode));
+	UE_LOG(LogOnline, Verbose, TEXT("3... Transaction: %s"), *InTransactionData.ToDebugString());
+	EPurchaseTransactionState Result = ConvertGPResponseCodeToPurchaseTransactionState(InResponseCode);
+
+	FString UserIdStr = GOOGLEPLAYUSER;
+	const TSharedRef<FOnlinePurchasePendingTransactionGooglePlay>* UserPendingTransactionPtr = PendingTransactions.Find(UserIdStr);
+	if(UserPendingTransactionPtr != nullptr)
+	{
+		const TSharedRef<FOnlinePurchasePendingTransactionGooglePlay> UserPendingTransaction = *UserPendingTransactionPtr;
+		const FString& ErrorStr = InTransactionData.GetErrorStr();
+
+		if (Result == EPurchaseTransactionState::Purchased || Result == EPurchaseTransactionState::Restored)
+		{
+			if (!UserPendingTransaction->AddCompletedOffer(Result, InTransactionData))
+			{
+				UE_LOG(LogOnline, Verbose, TEXT("Offer not found in pending transaction"));
+				Result = EPurchaseTransactionState::Failed;
+			}
+		}
+
+		UserPendingTransaction->PendingPurchaseInfo.TransactionState = Result;
+		UserPendingTransaction->PendingPurchaseInfo.TransactionId = InTransactionData.GetTransactionIdentifier();
+
+		FOnlineError FinalResult;
+		switch (Result)
+		{
+			case EPurchaseTransactionState::Failed:
+				FinalResult.SetFromErrorCode(TEXT("com.epicgames.purchase.failure"));
+				FinalResult.ErrorMessage = !ErrorStr.IsEmpty() ? FText::FromString(ErrorStr) : LOCTEXT("GooglePlayTransactionFailed", "Transaction Failed");
+				break;
+			case EPurchaseTransactionState::Canceled:
+				FinalResult.SetFromErrorCode(TEXT("com.epicgames.catalog_helper.user_cancelled"));
+				FinalResult.ErrorMessage = !ErrorStr.IsEmpty() ? FText::FromString(ErrorStr) : LOCTEXT("GooglePlayTransactionCancel", "Transaction Canceled");
+				break;
+			case EPurchaseTransactionState::Purchased:
+				FinalResult.bSucceeded = true;
+				break;
+			case EPurchaseTransactionState::Invalid:
+				FinalResult.SetFromErrorCode(TEXT("com.epicgames.purchase.invalid"));
+				FinalResult.ErrorMessage = !ErrorStr.IsEmpty() ? FText::FromString(ErrorStr) : LOCTEXT("GooglePlayInvalidState", "Invalid purchase result");
+				UserPendingTransaction->PendingPurchaseInfo.TransactionState = EPurchaseTransactionState::Invalid;
+				break;
+			default:
+				UE_LOG(LogOnline, Warning, TEXT("Unexpected state after purchase %d"), Result);
+				FinalResult.SetFromErrorCode(TEXT("com.epicgames.purchase.unexpected_state"));
+				FinalResult.ErrorMessage = !ErrorStr.IsEmpty() ? FText::FromString(ErrorStr) : LOCTEXT("GooglePlayUnexpectedState", "Unexpected purchase result");
+				UserPendingTransaction->PendingPurchaseInfo.TransactionState = EPurchaseTransactionState::Failed;
+				break;
+		}
+
+		TSharedRef<FPurchaseReceipt> FinalReceipt = UserPendingTransaction->GenerateReceipt();
+
+		TArray< TSharedRef<FPurchaseReceipt> >& UserCompletedTransactions = CompletedTransactions.FindOrAdd(UserIdStr);
+
+		PendingTransactions.Remove(UserIdStr);
+		UserCompletedTransactions.Add(FinalReceipt);
+
+		UserPendingTransaction->CheckoutCompleteDelegate.ExecuteIfBound(FinalResult, FinalReceipt);
+	}
+	else
+	{
+		UE_LOG(LogOnline, Log, TEXT("No pending transaction found associated with this purchase completion event"));
+	}
+}
+
+void FOnlinePurchaseGooglePlay::OnQueryExistingPurchasesComplete(EGooglePlayBillingResponseCode InResponseCode, const TArray<FGoogleTransactionData>& InExistingPurchases)
+{
+	UE_LOG(LogOnline, Verbose, TEXT("FOnlinePurchaseGooglePlay::OnQueryExistingPurchasesComplete Response: %s Num: %d"), ToString(InResponseCode), InExistingPurchases.Num());
+
+	if (bQueryingReceipts)
+	{
+		bool bSuccess = (InResponseCode == EGooglePlayBillingResponseCode::Ok);
+		if (bSuccess)
+		{
+			EPurchaseTransactionState Result = ConvertGPResponseCodeToPurchaseTransactionState(InResponseCode);
+			for (const FGoogleTransactionData& Purchase : InExistingPurchases)
+			{
+				UE_LOG(LogOnline, Verbose, TEXT("Adding existing receipt %s"), *Purchase.ToDebugString());
+				TSharedRef<FPurchaseReceipt> OfflineReceipt = FOnlinePurchasePendingTransactionGooglePlay::GenerateReceipt(Purchase);
+				OfflineTransactions.Add(OfflineReceipt);
+			}
+		}
+		else
+		{
+			UE_LOG(LogOnline, Verbose, TEXT("OnQueryExistingPurchasesComplete failed"));
+		}
+
+		Subsystem->ExecuteNextTick([this, bSuccess]()
+		{
+			FOnlineError Result(bSuccess);
+			QueryReceiptsComplete.ExecuteIfBound(Result);
+			bQueryingReceipts = false;
+			QueryReceiptsComplete.Unbind();
+		});
+	}
+	else
+	{
+		UE_LOG(LogOnline, Warning, TEXT("FOnlinePurchaseGooglePlay::OnQueryExistingPurchasesComplete unexpected call"));
+	}
+}
+
+TSharedRef<FPurchaseReceipt> FOnlinePurchasePendingTransactionGooglePlay::GenerateReceipt()
+{
+	TSharedRef<FPurchaseReceipt> Receipt = MakeShareable(new FPurchaseReceipt());
+	
+	Receipt->TransactionState = PendingPurchaseInfo.TransactionState;
+	Receipt->TransactionId = PendingPurchaseInfo.TransactionId;
+	
+	if(PendingPurchaseInfo.TransactionState == EPurchaseTransactionState::Purchased ||
+	   PendingPurchaseInfo.TransactionState == EPurchaseTransactionState::Restored)
+	{
+		Receipt->ReceiptOffers = PendingPurchaseInfo.ReceiptOffers;
+	}
+	else
+	{
+		// Add the requested offers to the receipt in the event of an incomplete purchase.
+		for(const auto& RequestedOffer : CheckoutRequest.PurchaseOffers)
+		{
+			Receipt->AddReceiptOffer(RequestedOffer.OfferNamespace, RequestedOffer.OfferId, RequestedOffer.Quantity);
+		}
+	}
+	
+	return Receipt;
+}
+
+TSharedRef<FPurchaseReceipt> FOnlinePurchasePendingTransactionGooglePlay::GenerateReceipt(const FGoogleTransactionData& Transaction)
+{
+	TSharedRef<FPurchaseReceipt> Receipt = MakeShareable(new FPurchaseReceipt());
+	
+	Receipt->TransactionState = Transaction.GetErrorStr().IsEmpty() ? EPurchaseTransactionState::Purchased : EPurchaseTransactionState::Failed;
+	Receipt->TransactionId = Transaction.GetTransactionIdentifier();
+	
+	if (Receipt->TransactionState == EPurchaseTransactionState::Purchased ||
+		Receipt->TransactionState == EPurchaseTransactionState::Restored)
+	{
+		FPurchaseReceipt::FReceiptOfferEntry ReceiptEntry(TEXT(""), Transaction.GetOfferId(), 1);
+
+		int32 Idx = ReceiptEntry.LineItems.AddZeroed();
+
+		FPurchaseReceipt::FLineItemInfo& LineItem = ReceiptEntry.LineItems[Idx];
+
+		LineItem.ItemName = Transaction.GetOfferId();
+		LineItem.UniqueId = Transaction.GetTransactionIdentifier();
+		LineItem.ValidationInfo = Transaction.GetCombinedReceiptData();
+
+		Receipt->AddReceiptOffer(ReceiptEntry);
+	}
+	
+	return Receipt;
+}
+
+bool FOnlinePurchasePendingTransactionGooglePlay::AddCompletedOffer(EPurchaseTransactionState Result, const FGoogleTransactionData& Transaction)
+{
+	for (int32 OfferIdx = 0; OfferIdx < CheckoutRequest.PurchaseOffers.Num(); ++OfferIdx)
+	{
+		const FPurchaseCheckoutRequest::FPurchaseOfferEntry& Offer = CheckoutRequest.PurchaseOffers[OfferIdx];
+		if (Transaction.GetOfferId() == Offer.OfferId)
+		{
+			FPurchaseReceipt::FReceiptOfferEntry Receipt(TEXT(""), Transaction.GetOfferId(), 1);
+
+			int32 Idx = Receipt.LineItems.AddZeroed();
+			
+			FPurchaseReceipt::FLineItemInfo& LineItem = Receipt.LineItems[Idx];
+			
+			LineItem.ItemName = Transaction.GetOfferId();
+			LineItem.UniqueId = Transaction.GetTransactionIdentifier();
+			LineItem.ValidationInfo = Transaction.GetCombinedReceiptData();
+			
+			PendingPurchaseInfo.AddReceiptOffer(Receipt);
+			return true;
+		}
+	}
+	
+	return false;
+}
+
+#undef LOCTEXT_NAMESPACE
\ No newline at end of file
diff --git a/Engine/Plugins/Online/Android/OnlineSubsystemGooglePlay/Source/Private/OnlineStoreGooglePlay.cpp b/Engine/Plugins/Online/Android/OnlineSubsystemGooglePlay/Source/Private/OnlineStoreGooglePlay.cpp
new file mode 100644
--- /dev/null
+++ b/Engine/Plugins/Online/Android/OnlineSubsystemGooglePlay/Source/Private/OnlineStoreGooglePlay.cpp
@@ -0,0 +1,175 @@
+// Copyright 1998-2016 Epic Games, Inc. All Rights Reserved.
+
+#include "OnlineStoreGooglePlay.h"
+#include "OnlineSubsystemGooglePlay.h"
+#include "OnlineAsyncTaskGooglePlayQueryInAppPurchases.h"
+#include "Misc/ConfigCacheIni.h"
+
+FOnlineStoreGooglePlayV2::FOnlineStoreGooglePlayV2(FOnlineSubsystemGooglePlay* InSubsystem)
+	: bIsQueryInFlight(false)
+	, Subsystem(InSubsystem)
+{
+	UE_LOG(LogOnline, Verbose, TEXT( "FOnlineStoreGooglePlayV2::FOnlineStoreGooglePlayV2" ));
+}
+
+FOnlineStoreGooglePlayV2::FOnlineStoreGooglePlayV2()
+	: bIsQueryInFlight(false)
+	, Subsystem(nullptr)
+{
+	UE_LOG(LogOnline, Verbose, TEXT( "FOnlineStoreGooglePlayV2::FOnlineStoreGooglePlayV2 empty" ));
+}
+
+FOnlineStoreGooglePlayV2::~FOnlineStoreGooglePlayV2()
+{
+	if (Subsystem)
+	{
+		Subsystem->ClearOnGooglePlayAvailableIAPQueryCompleteDelegate_Handle(AvailableIAPQueryDelegateHandle);
+	}
+}
+
+void FOnlineStoreGooglePlayV2::Init()
+{
+	UE_LOG(LogOnline, Verbose, TEXT("FOnlineStoreGooglePlayV2::Init"));
+
+	FOnGooglePlayAvailableIAPQueryCompleteDelegate Delegate = FOnGooglePlayAvailableIAPQueryCompleteDelegate::CreateThreadSafeSP(this, &FOnlineStoreGooglePlayV2::OnGooglePlayAvailableIAPQueryComplete);
+	AvailableIAPQueryDelegateHandle = Subsystem->AddOnGooglePlayAvailableIAPQueryCompleteDelegate_Handle(Delegate);
+
+	FString GooglePlayLicenseKey;
+	if (!GConfig->GetString(TEXT("/Script/AndroidRuntimeSettings.AndroidRuntimeSettings"), TEXT("GooglePlayLicenseKey"), GooglePlayLicenseKey, GEngineIni) || GooglePlayLicenseKey.IsEmpty())
+	{
+		UE_LOG_ONLINE(Warning, TEXT("Missing GooglePlayLicenseKey key in /Script/AndroidRuntimeSettings.AndroidRuntimeSettings of DefaultEngine.ini"));
+	}
+
+	extern void AndroidThunkCpp_Iap_SetupIapService(const FString&);
+	AndroidThunkCpp_Iap_SetupIapService(GooglePlayLicenseKey);
+}
+
+TSharedRef<FOnlineStoreOffer> ConvertProductToStoreOffer(const FInAppPurchaseProductInfo& Product)
+{
+	TSharedRef<FOnlineStoreOffer> NewProductInfo = MakeShareable(new FOnlineStoreOffer());
+
+	NewProductInfo->OfferId = Product.Identifier;
+
+	FString Title = Product.DisplayName;
+	int32 OpenParenIdx = -1;
+	int32 CloseParenIdx = -1;
+	if (Title.FindLastChar(TEXT(')'), CloseParenIdx) && Title.FindLastChar(TEXT('('), OpenParenIdx) && (OpenParenIdx < CloseParenIdx))
+	{
+		Title = Title.Left(OpenParenIdx).TrimTrailing();
+	}
+
+	NewProductInfo->Title = FText::FromString(Title);
+	NewProductInfo->Description = FText::FromString(Product.DisplayDescription); // google has only one description, map it to (short) description to match iOS
+	//NewProductInfo->LongDescription = FText::FromString(Product.DisplayDescription); // leave this empty so we know it's not set (client can apply more info from MCP)
+	NewProductInfo->PriceText = FText::FromString(Product.DisplayPrice);
+	NewProductInfo->CurrencyCode = Product.CurrencyCode;
+
+	NewProductInfo->NumericPrice = FMath::TruncToInt(Product.RawPrice);
+
+	// Google doesn't support these fields, set to min and max defaults
+	NewProductInfo->ReleaseDate = FDateTime::MinValue();
+	NewProductInfo->ExpirationDate = FDateTime::MaxValue();
+
+	return NewProductInfo;
+}
+
+void FOnlineStoreGooglePlayV2::OnGooglePlayAvailableIAPQueryComplete(EGooglePlayBillingResponseCode InResponseCode, const TArray<FInAppPurchaseProductInfo>& InProvidedProductInformation)
+{ 
+	TArray<FUniqueOfferId> OfferIds;
+	for (const FInAppPurchaseProductInfo& Product : InProvidedProductInformation)
+	{
+		TSharedRef<FOnlineStoreOffer> NewProductOffer = ConvertProductToStoreOffer(Product);
+
+		AddOffer(NewProductOffer);
+		OfferIds.Add(NewProductOffer->OfferId);
+
+		UE_LOG(LogOnline, Log, TEXT("Product Identifier: %s, Name: %s, Desc: %s, Long Desc: %s, Price: %s IntPrice: %d"),
+			*NewProductOffer->OfferId,
+			*NewProductOffer->Title.ToString(),
+			*NewProductOffer->Description.ToString(),
+			*NewProductOffer->LongDescription.ToString(),
+			*NewProductOffer->PriceText.ToString(),
+			NewProductOffer->NumericPrice);
+	}
+
+	CurrentQueryTask->ProcessQueryAvailablePurchasesResults(InResponseCode);
+
+	// clear the pointer, it will be destroyed by the async task manager
+	CurrentQueryTask = nullptr;
+	bIsQueryInFlight = false;
+}
+
+void FOnlineStoreGooglePlayV2::QueryCategories(const FUniqueNetId& UserId, const FOnQueryOnlineStoreCategoriesComplete& Delegate)
+{
+	Delegate.ExecuteIfBound(false, TEXT("No CatalogService"));
+}
+
+void FOnlineStoreGooglePlayV2::GetCategories(TArray<FOnlineStoreCategory>& OutCategories) const
+{
+	OutCategories.Empty();
+}
+
+void FOnlineStoreGooglePlayV2::QueryOffersByFilter(const FUniqueNetId& UserId, const FOnlineStoreFilter& Filter, const FOnQueryOnlineStoreOffersComplete& Delegate)
+{
+	Delegate.ExecuteIfBound(false, TArray<FUniqueOfferId>(), TEXT("No CatalogService"));
+}
+
+void FOnlineStoreGooglePlayV2::QueryOffersById(const FUniqueNetId& UserId, const TArray<FUniqueOfferId>& OfferIds, const FOnQueryOnlineStoreOffersComplete& Delegate)
+{
+	UE_LOG(LogOnline, Verbose, TEXT("FOnlineStoreGooglePlayV2::QueryOffersById"));
+
+	if (bIsQueryInFlight)
+	{
+		Delegate.ExecuteIfBound(false, OfferIds, TEXT("Request already in flight"));
+	}
+	else if (OfferIds.Num() == 0)
+	{
+		Delegate.ExecuteIfBound(false, OfferIds, TEXT("No offers to query for"));
+	}
+	else
+	{
+		CurrentQueryTask = new FOnlineAsyncTaskGooglePlayQueryInAppPurchasesV2(
+			Subsystem,
+			OfferIds,
+			Delegate);
+		Subsystem->QueueAsyncTask(CurrentQueryTask);
+		
+		bIsQueryInFlight = true;
+	}
+}
+
+void FOnlineStoreGooglePlayV2::AddOffer(const TSharedRef<FOnlineStoreOffer>& NewOffer)
+{
+	TSharedRef<FOnlineStoreOffer>* Existing = CachedOffers.Find(NewOffer->OfferId);
+	if (Existing != nullptr)
+	{
+		// Replace existing offer
+		*Existing = NewOffer;
+	}
+	else
+	{
+		CachedOffers.Add(NewOffer->OfferId, NewOffer);
+	}
+}
+
+void FOnlineStoreGooglePlayV2::GetOffers(TArray<FOnlineStoreOfferRef>& OutOffers) const
+{
+	for (const auto& CachedEntry : CachedOffers)
+	{
+		const TSharedRef<FOnlineStoreOffer>& CachedOffer = CachedEntry.Value;
+		OutOffers.Add(CachedOffer);
+	}
+}
+
+TSharedPtr<FOnlineStoreOffer> FOnlineStoreGooglePlayV2::GetOffer(const FUniqueOfferId& OfferId) const
+{
+	TSharedPtr<FOnlineStoreOffer> Result;
+
+	const TSharedRef<FOnlineStoreOffer>* Existing = CachedOffers.Find(OfferId);
+	if (Existing != nullptr)
+	{
+		Result = (*Existing);
+	}
+
+	return Result;
+}
diff --git a/Engine/Plugins/Online/Android/OnlineSubsystemGooglePlay/Source/Private/OnlineStoreGooglePlayCommon.h b/Engine/Plugins/Online/Android/OnlineSubsystemGooglePlay/Source/Private/OnlineStoreGooglePlayCommon.h
new file mode 100644
--- /dev/null
+++ b/Engine/Plugins/Online/Android/OnlineSubsystemGooglePlay/Source/Private/OnlineStoreGooglePlayCommon.h
@@ -0,0 +1,155 @@
+// Copyright 1998-2016 Epic Games, Inc. All Rights Reserved.
+
+#pragma once
+
+#include "OnlineStoreInterface.h"
+#include "OnlineStoreInterfaceV2.h"
+#include "OnlinePurchaseInterface.h"
+
+struct FGoogleTransactionData;
+
+/** Possible responses returned from the Java GooglePlay billing interface */
+enum class EGooglePlayBillingResponseCode : uint8
+{
+	Ok = 0,
+	UserCancelled = 1,
+	ServiceUnavailable = 2,
+	BillingUnavailable = 3,
+	ItemUnavailable = 4,
+	DeveloperError = 5,
+	Error = 6,
+	ItemAlreadyOwned = 7,
+	ItemNotOwned = 8,
+};
+
+inline const TCHAR* const ToString(EGooglePlayBillingResponseCode InResponseCode)
+{
+	switch (InResponseCode)
+	{
+		case EGooglePlayBillingResponseCode::Ok:
+			return TEXT("Ok");
+		case EGooglePlayBillingResponseCode::UserCancelled:
+			return TEXT("UserCancelled");
+		case EGooglePlayBillingResponseCode::ItemAlreadyOwned:
+			return TEXT("ItemAlreadyOwned");
+		case EGooglePlayBillingResponseCode::ItemNotOwned:
+			return TEXT("ItemNotOwned");
+		case EGooglePlayBillingResponseCode::ServiceUnavailable:
+			return TEXT("ServiceUnavailable");
+		case EGooglePlayBillingResponseCode::BillingUnavailable:
+			return TEXT("BillingUnavailable");
+		case EGooglePlayBillingResponseCode::ItemUnavailable:
+			return TEXT("ItemUnavailable");
+		case EGooglePlayBillingResponseCode::DeveloperError:
+			return TEXT("DeveloperError");
+		case EGooglePlayBillingResponseCode::Error:
+			return TEXT("Error");
+		default:
+			return TEXT("UnknownError");
+	}
+}
+
+inline EInAppPurchaseState::Type ConvertGPResponseCodeToIAPState(const EGooglePlayBillingResponseCode InResponseCode)
+{
+	switch (InResponseCode)
+	{
+		case EGooglePlayBillingResponseCode::Ok:
+			return EInAppPurchaseState::Success;
+		case EGooglePlayBillingResponseCode::UserCancelled:
+			return EInAppPurchaseState::Cancelled;
+		case EGooglePlayBillingResponseCode::ItemAlreadyOwned:
+			return EInAppPurchaseState::AlreadyOwned;
+		case EGooglePlayBillingResponseCode::ItemNotOwned:
+			return EInAppPurchaseState::NotAllowed;
+		case EGooglePlayBillingResponseCode::ServiceUnavailable:
+		case EGooglePlayBillingResponseCode::BillingUnavailable:
+		case EGooglePlayBillingResponseCode::ItemUnavailable:
+		case EGooglePlayBillingResponseCode::DeveloperError:
+		case EGooglePlayBillingResponseCode::Error:
+		default:
+			return EInAppPurchaseState::Failed;
+	}
+}
+
+inline EPurchaseTransactionState ConvertGPResponseCodeToPurchaseTransactionState(const EGooglePlayBillingResponseCode InResponseCode)
+{
+	switch (InResponseCode)
+	{
+		case EGooglePlayBillingResponseCode::Ok:
+			return EPurchaseTransactionState::Purchased;
+		case EGooglePlayBillingResponseCode::UserCancelled:
+			return EPurchaseTransactionState::Canceled;
+		case EGooglePlayBillingResponseCode::ItemAlreadyOwned:
+			// Non consumable purchased again?
+			return EPurchaseTransactionState::Invalid;
+		case EGooglePlayBillingResponseCode::ItemNotOwned:
+			return EPurchaseTransactionState::Invalid;
+		case EGooglePlayBillingResponseCode::ServiceUnavailable:
+		case EGooglePlayBillingResponseCode::BillingUnavailable:
+		case EGooglePlayBillingResponseCode::ItemUnavailable:
+		case EGooglePlayBillingResponseCode::DeveloperError:
+		case EGooglePlayBillingResponseCode::Error:
+		default:
+			return EPurchaseTransactionState::Failed;
+	}
+}
+
+/**
+ * The resulting state of an iap transaction
+ */
+enum class EInAppPurchaseResult : uint8
+{
+	Succeeded = 0,
+	RestoredFromServer,
+	Failed,
+	Cancelled,
+};
+
+/**
+ * Implementation of the Platform Purchase receipt. For this we provide an identifier and the encrypted data.
+ */
+class FGooglePlayPurchaseReceipt : public IPlatformPurchaseReceipt
+{
+public:
+	// Product identifier
+	FString Identifier;
+
+	// The encrypted receipt data
+	FString Data;
+};
+
+ /**
+  * Delegate fired when an IAP query for available offers has completed
+  *
+  * @param Response response from GooglePlay backend
+  * @param ProvidedProductInformation list of offers returned in response to a query on available offer ids
+  */
+DECLARE_MULTICAST_DELEGATE_TwoParams(FOnGooglePlayAvailableIAPQueryComplete, EGooglePlayBillingResponseCode /*Response*/, const TArray<FInAppPurchaseProductInfo>& /*ProvidedProductInformation*/);
+typedef FOnGooglePlayAvailableIAPQueryComplete::FDelegate FOnGooglePlayAvailableIAPQueryCompleteDelegate;
+
+/**
+ * Delegate fired when an IAP has completed
+ *
+ * @param InResponseCode response from the GooglePlay backend
+ * @param InTransactionData transaction data for the completed purchase
+ */
+DECLARE_MULTICAST_DELEGATE_TwoParams(FOnGooglePlayProcessPurchaseComplete, EGooglePlayBillingResponseCode /*InResponseCode*/, const FGoogleTransactionData& /*InTransactionData*/);
+typedef FOnGooglePlayProcessPurchaseComplete::FDelegate FOnGooglePlayProcessPurchaseCompleteDelegate;
+
+/**
+ * Delegate fired internally an existing purchases query has completed
+ *
+ * @param InResponseCode response from the GooglePlay backend
+ * @param InExistingPurchases known purchases for the user (non consumed or permanent)
+ */
+DECLARE_MULTICAST_DELEGATE_TwoParams(FOnGooglePlayQueryExistingPurchasesComplete, EGooglePlayBillingResponseCode /*InResponseCode*/, const TArray<FGoogleTransactionData>& /*InExistingPurchases*/);
+typedef FOnGooglePlayQueryExistingPurchasesComplete::FDelegate FOnGooglePlayQueryExistingPurchasesCompleteDelegate;
+
+/**
+ * Delegate fired internally when existing purchases have been restored (StoreV1 only)
+ *
+ * @param InResponseCode response from the GooglePlay backend
+ * @param InRestoredPurchases restored for the user (non consumed or permanent)
+ */
+DECLARE_MULTICAST_DELEGATE_TwoParams(FOnGooglePlayRestorePurchasesComplete, EGooglePlayBillingResponseCode /*InResponseCode*/, const TArray<FGoogleTransactionData>& /*InRestoredPurchases*/);
+typedef FOnGooglePlayRestorePurchasesComplete::FDelegate FOnGooglePlayRestorePurchasesCompleteDelegate;
