Copying //UE4/WEX-Staging/... to //UE4/Main (Source: //WEX/Main @ 3580612)

=================================================================================================
MAJOR FEATURES + CHANGES
=================================================================================================

Change 3526838 by [REDACTED_PERSON]

	#WEX: make map elements support TitleProperty on their values in the editor

Change 3517937 by [REDACTED_PERSON]

	Replace with AddChunkInstallDelegate, which supports a bool error code and is bound once instead of separately for each chunk. All implementations support this delegate at a basic level, although several could be improved to call the failure delegate in more cases

Change 3498765 by [REDACTED_PERSON]

	#WEX: Added a way to bind a delegate that fires whenever an analytics event is queued.
	- Bind this delegate and use it to log analytics events (for now)

Change 3495796 by [REDACTED_PERSON]

	#WEX
	- Reworked the LoadTimes.DumpReport console command to accept command arguments and added options for alphanumeric sorting (-alphasort), tweakable asset time cutoff (lowtime=X), and file output to the Saved/Profiling directory (file).
	- Added hooks for automatically generating load time reports for every map load (enabled using the DUMP_LOAD_REPORT_PER_MAP #define).

Change 3489241 by [REDACTED_PERSON]

	#UE4 - First unfinished pass to GoogleIOS
	- SDK auth token data needs to copy auth into TMap properly

Change 3487767 by [REDACTED_PERSON]

	#Analytics: Make FAnalyticsEventAttribute support typed values
	- This makes sure the value types in the resultant JSON reflect the code.
	- Added support for Number (double), Boolean, Null, and JsonFragment types
	- This should make it so we don't have to whitelist everything to be converted to number on the Grafana processing side.
	- Made all attributes on FAnalyticsEventAttribute immutable

Change 3478818 by [REDACTED_PERSON]

	Add detection of Houdini (running on Intel Android CPU)
	#ue4
	#android

Change 3475449 by [REDACTED_PERSON]

	Add disable force inline option for iOS build, enabled for WEX.

diff --git a/Engine/Build/Android/Java/src/com/epicgames/ue4/GameActivity.java b/Engine/Build/Android/Java/src/com/epicgames/ue4/GameActivity.java
--- a/Engine/Build/Android/Java/src/com/epicgames/ue4/GameActivity.java
+++ b/Engine/Build/Android/Java/src/com/epicgames/ue4/GameActivity.java
@@ -19,9 +19,7 @@ import android.util.Log;
 
 import android.os.Vibrator;
 import android.os.SystemClock;
-//@HSL_BEGIN - chance.lyon - Adding Handler for Sticky window callbacks
 import android.os.Handler;
-//@HSL_END
 
 import android.app.AlertDialog;
 import android.app.Dialog;
@@ -57,9 +55,7 @@ import android.media.AudioManager;
 import android.util.DisplayMetrics;
 
 import android.view.InputDevice;
-//@HSL_BEGIN - chance.lyon - Adding KeyEvent for Sticky window callbacks
 import android.view.KeyEvent;
-//@HSL_END
 import android.view.Gravity;
 import android.view.MotionEvent;
 import android.view.View;
@@ -125,6 +121,13 @@ import android.view.inputmethod.InputMethodManager;
 import android.graphics.Rect;
 import android.view.ViewTreeObserver;
 
+import android.widget.FrameLayout;
+import android.widget.RelativeLayout;
+import android.view.inputmethod.InputConnection;
+import android.view.inputmethod.InputConnectionWrapper;
+import android.util.AttributeSet;
+import android.graphics.Color;
+
 //$${gameActivityImportAdditions}$$
 
 //$${gameActivityPostImportAdditions}$$
@@ -165,7 +168,7 @@ public class GameActivity extends NativeActivity implements SurfaceHolder.Callba
 	private static final String CONSOLE_SPINNER_ITEMS[] = {"Common Console Commands", "stat FPS", "stat Anim","stat OpenGLRHI","stat VulkanRHI","stat DumpEvents","stat DumpFrame",
 		"stat DumpHitches","stat Engine","stat Game","stat Grouped","stat Hitches","stat InitViews","stat LightRendering",
 		"stat Memory","stat Particles","stat SceneRendering","stat SceneUpdate","stat ShadowRendering","stat Slow",
-		"stat Streaming","stat StreamingDetails","stat Unit","stat UnitGraph", "stat StartFile", "stat StopFile", "GameVer", "show PostProcessing"};
+		"stat Streaming","stat StreamingDetails","stat Unit","stat UnitGraph", "stat StartFile", "stat StopFile", "GameVer", "show PostProcessing", "stat AndroidCPU"};
 	AlertDialog consoleAlert;
 	LinearLayout consoleAlertLayout;
 	Spinner consoleSpinner;
@@ -180,6 +183,15 @@ public class GameActivity extends NativeActivity implements SurfaceHolder.Callba
 	EditText virtualKeyboardInputBox;
 	String virtualKeyboardPreviousContents;
 
+	//virtual keyboard input - custom EditText
+	VirtualKeyboardInput newVirtualKeyboardInput;
+
+	//layout for displaying virtual keyboard input
+	private LinearLayout virtualKeyboardLayout;
+
+	//container for SurfaceView and virtual keyboard input
+	private FrameLayout containerFrameLayout;
+
 	// Keep a reference to the main content view so we can bring up the virtual keyboard without an editbox
 	private View mainView;
 	private boolean bKeyboardShowing;
@@ -932,11 +944,17 @@ public class GameActivity extends NativeActivity implements SurfaceHolder.Callba
 			HasAllFiles = true;
 		}
 
+		containerFrameLayout = new FrameLayout(_activity);
+		virtualKeyboardLayout = new LinearLayout(_activity);
+
 		// Need to create our surface view here regardless of if we are going to end up using it
 		getWindow().takeSurface(null);
 		MySurfaceView = new SurfaceView(this);
+		MySurfaceView.setBackgroundColor(Color.TRANSPARENT);
 		MySurfaceView.getHolder().addCallback(this);
-		setContentView(MySurfaceView);
+		containerFrameLayout.addView(MySurfaceView);
+		containerFrameLayout.addView(virtualKeyboardLayout);
+		setContentView(containerFrameLayout);
 
 		// cache a reference to the main content view and set it so it can be focused on
         mainView = findViewById( android.R.id.content );
@@ -947,7 +965,9 @@ public class GameActivity extends NativeActivity implements SurfaceHolder.Callba
         mainDecorViewRect = new Rect();
         mainDecorView.getWindowVisibleDisplayFrame( mainDecorViewRect );
 
-        mainView.getViewTreeObserver().addOnGlobalLayoutListener( new ViewTreeObserver.OnGlobalLayoutListener()
+		createVirtualKeyboardInput();
+		
+		mainView.getViewTreeObserver().addOnGlobalLayoutListener( new ViewTreeObserver.OnGlobalLayoutListener()
         {
         	@Override
         	public void onGlobalLayout()
@@ -955,7 +975,8 @@ public class GameActivity extends NativeActivity implements SurfaceHolder.Callba
         		if( bKeyboardShowing )
         		{
         			Rect visibleRect = new Rect();
-        			View visibleView = getWindow().getDecorView();
+        			View visibleView = mainView.getRootView();
+
         			visibleView.getWindowVisibleDisplayFrame( visibleRect );
 
         			// determine which side of the screen the keyboard is covering
@@ -971,7 +992,22 @@ public class GameActivity extends NativeActivity implements SurfaceHolder.Callba
     				keyboardRect.right = ( leftDiff > 0 ) ? visibleRect.left : mainDecorViewRect.right; // keyboard is on the left
     				keyboardRect.bottom = ( topDiff > 0 ) ? visibleRect.top : mainDecorViewRect.bottom; // keyboard is on the top
 
+    				int screenHeight = visibleView.getHeight();
+    				int heightDiff = visibleRect.bottom - newVirtualKeyboardInput.getHeight();
+                    
         			nativeVirtualKeyboardShown( keyboardRect.left, keyboardRect.top, keyboardRect.right, keyboardRect.bottom );
+                    if(bottomDiff + topDiff + rightDiff + topDiff> 100)
+                    {
+                    	newVirtualKeyboardInput.setY(heightDiff);
+                    	newVirtualKeyboardInput.setVisibility(View.VISIBLE);
+						newVirtualKeyboardInput.requestFocus();
+                    }
+                    else if(newVirtualKeyboardInput.getY() > 0)
+                    {
+						newVirtualKeyboardInput.setVisibility(View.GONE);
+						//set offscreen
+        				newVirtualKeyboardInput.setY(-1000); 
+                    }
         		}
         	}
         });
@@ -981,7 +1017,6 @@ public class GameActivity extends NativeActivity implements SurfaceHolder.Callba
 		Log.debug("==============> GameActive.onCreate complete!");
 	}
 
-//@HSL_BEGIN - chance.lyon - Adding Handler for Sticky window callbacks
 	private Handler mRestoreImmersiveModeHandler = new Handler();
 	private Runnable restoreImmersiveModeRunnable = new Runnable()
 	{
@@ -1037,7 +1072,6 @@ public class GameActivity extends NativeActivity implements SurfaceHolder.Callba
 
 		return super.onKeyDown(keyCode, event);
 	}
-//@HSL_END
 	
 	@Override
 	public void onResume()
@@ -1050,7 +1084,6 @@ public class GameActivity extends NativeActivity implements SurfaceHolder.Callba
 		// only do this on KitKat and above
 		if (ShouldHideUI)
 		{ 
-//@HSL_BEGIN - chance.lyon - Adding Handler for Sticky window callbacks
 			restoreTransparentBars();
 			View decorView = getWindow().getDecorView(); 
 
@@ -1072,7 +1105,6 @@ public class GameActivity extends NativeActivity implements SurfaceHolder.Callba
 				}
 			});
 		}
-//@HSL_END
 		
 		if(HasAllFiles)
 		{
@@ -1382,7 +1414,6 @@ public class GameActivity extends NativeActivity implements SurfaceHolder.Callba
 		// Set label and starting contents
 		virtualKeyboardAlert.setTitle(Label);
 
-		// @HSL_BEGIN - Josh.May - 11/01/2016 - Ensure the input mode of the text box is set before setting the contents.
 		// configure for type of input
 		virtualKeyboardInputBox.setRawInputType(inInputType);
 		virtualKeyboardInputBox.setTransformationMethod((inInputType & InputType.TYPE_TEXT_VARIATION_PASSWORD) == 0 ? null : PasswordTransformationMethod.getInstance());
@@ -1390,7 +1421,6 @@ public class GameActivity extends NativeActivity implements SurfaceHolder.Callba
 		virtualKeyboardInputBox.setText("");
 		virtualKeyboardInputBox.append(Contents);
 		virtualKeyboardPreviousContents = Contents;
-		// @HSL_END - Josh.May - 11/01/2016
 
 		_activity.runOnUiThread(new Runnable()
 		{
@@ -1413,36 +1443,63 @@ public class GameActivity extends NativeActivity implements SurfaceHolder.Callba
 		{
 			public void run()
 			{
+				Log.debug("Hide newVirtualKeyboardInput");
+
 		        InputMethodManager imm =(InputMethodManager)getSystemService(Context.INPUT_METHOD_SERVICE);
-		        imm.hideSoftInputFromWindow(mainView.getWindowToken(), 0);
+				imm.hideSoftInputFromWindow(newVirtualKeyboardInput.getWindowToken(), 0);
+
+				newVirtualKeyboardInput.clearFocus();
+				//set offscreen
+				newVirtualKeyboardInput.setY(-1000);
 
+				newVirtualKeyboardInput.setVisibility(View.GONE);
+
+				nativeVirtualKeyboardVisible(false);
 		        bKeyboardShowing = false;
 			}
 		});
 	}
 
-	public void AndroidThunkJava_ShowVirtualKeyboardInput(int InputType, String Label, String Contents)
+	//initial settings for the virtual input
+	String virtualKeyboardInputContent;
+	int virtualKeyboardInputType;
+	public void AndroidThunkJava_ShowVirtualKeyboardInput(int inInputType, String Label, String Contents)
 	{
+		virtualKeyboardInputContent = Contents;
+		virtualKeyboardInputType = inInputType;
 		_activity.runOnUiThread(new Runnable()
 		{
 			public void run()
 			{
+				newVirtualKeyboardInput.setVisibility(View.VISIBLE);
+				//set offscreen
+				newVirtualKeyboardInput.setY(-1000);
+
+				newVirtualKeyboardInput.setText("");
+				newVirtualKeyboardInput.append(virtualKeyboardInputContent);
+
 				if (bKeyboardShowing)
 				{
 					Log.debug("Virtual keyboard already showing.");
-					return;
 				}
-				else if (mainView.requestFocus())
+				else if(newVirtualKeyboardInput.requestFocus())
 				{
-		            InputMethodManager imm = (InputMethodManager) getSystemService(Context.INPUT_METHOD_SERVICE);
-		            imm.showSoftInput(mainView, 0);
+					InputMethodManager imm = (InputMethodManager) getSystemService(Context.INPUT_METHOD_SERVICE);
 
-		            bKeyboardShowing = true;
-		        }
-		    }
-	    });
-		// @HSL_END - Josh.May - 3/28/2017
- 	}
+					imm.showSoftInput(newVirtualKeyboardInput, 0);
+
+					newVirtualKeyboardInput.setRawInputType(virtualKeyboardInputType);
+					newVirtualKeyboardInput.setTransformationMethod((
+						virtualKeyboardInputType & InputType.TYPE_TEXT_VARIATION_PASSWORD) == 0 ? 
+						null : 
+						PasswordTransformationMethod.getInstance());
+
+					nativeVirtualKeyboardVisible(true);
+					bKeyboardShowing = true;
+				}
+			}
+		});
+	}
 	
 	public void AndroidThunkJava_LaunchURL(String URL)
 	{
@@ -2611,6 +2668,169 @@ public class GameActivity extends NativeActivity implements SurfaceHolder.Callba
 		}
 		return _bundle.getString(key);
 	}
+			
+	//virtual keyboard input class - custom EditText
+	public class VirtualKeyboardInput extends EditText 
+	{
+		public VirtualKeyboardInput(Context context, AttributeSet attrs, int defStyle) 
+		{
+			super(context, attrs, defStyle);
+		}
+
+		public VirtualKeyboardInput(Context context, AttributeSet attrs) 
+		{
+			super(context, attrs);
+		}
+
+		public VirtualKeyboardInput(Context context) 
+		{
+			super(context);
+		}
+
+		//Override BACK key to hide the virtual keyboard
+		@Override 
+		public boolean onKeyPreIme(int keyCode, KeyEvent event) 
+		{
+		    if (keyCode == KeyEvent.KEYCODE_BACK &&  event.getAction() == KeyEvent.ACTION_UP) 
+			{
+		    	AndroidThunkJava_HideVirtualKeyboardInput();
+		        return false;
+		    }
+		    return super.dispatchKeyEvent(event);
+		}
+
+		//extend associated InputConnection to handle special keys 
+		@Override
+		public InputConnection onCreateInputConnection(EditorInfo outAttrs) {
+			return new VirtualKeyboardInputConnection(super.onCreateInputConnection(outAttrs),
+					true, this);
+		}
+
+		private class VirtualKeyboardInputConnection extends InputConnectionWrapper {
+			VirtualKeyboardInput owner;
+			public VirtualKeyboardInputConnection(InputConnection target, boolean mutable, VirtualKeyboardInput editText) {
+						super(target, mutable);
+						owner = editText;
+					}
+
+			//implement virtual keyboard's BACKSPACE and ENTER keys
+			@Override
+			public boolean sendKeyEvent(KeyEvent event) 
+			{
+				StringBuffer text = new StringBuffer(owner.getText().toString());
+
+				int selStart, selEnd;
+ 
+				int a = owner.getSelectionStart();
+				int b = owner.getSelectionEnd();
+ 
+				selStart = Math.min(a, b);
+				selEnd = Math.max(a, b);
+				
+				if (event.getAction() == KeyEvent.ACTION_DOWN && event.getKeyCode() == KeyEvent.KEYCODE_DEL) 
+				{
+					//delete selected text / previous character
+					if (selStart != selEnd) 
+					{
+						text.replace(selStart, selEnd, "");
+					} 
+					else if(selStart > 0)
+					{ 
+						selStart--;
+						text.replace(selStart, selStart + 1, "");
+					} 
+					owner.setText(text.toString());
+					owner.setSelection(selStart);
+					return true;
+				}
+				else if (event.getAction() == KeyEvent.ACTION_DOWN && event.getKeyCode() == KeyEvent.KEYCODE_ENTER) 
+				{
+					//add new line
+					if (selStart != selEnd) 
+					{
+						text.replace(selStart, selEnd, "\n");
+					} 
+					else 
+					{ 
+						text.insert(selStart, "\n");
+					} 
+					owner.setText(text.toString());
+					owner.setSelection(selStart + 1);
+					return true;
+				}
+				return super.sendKeyEvent(event);
+			}
+
+			//in latest Android, deleteSurroundingText(1, 0) will be called for BACKSPACE
+			@Override
+			public boolean deleteSurroundingText(int beforeLength, int afterLength) {       
+				if (beforeLength == 1 && afterLength == 0) {
+					// backspace
+					return sendKeyEvent(new KeyEvent(KeyEvent.ACTION_DOWN, KeyEvent.KEYCODE_DEL))
+						&& sendKeyEvent(new KeyEvent(KeyEvent.ACTION_UP, KeyEvent.KEYCODE_DEL));
+				}
+
+				return super.deleteSurroundingText(beforeLength, afterLength);
+			}
+
+		}
+	}
+
+	//create the virtual keyboard input 
+	private void createVirtualKeyboardInput()
+	{
+		newVirtualKeyboardInput = new VirtualKeyboardInput(this);
+		newVirtualKeyboardInput.setSingleLine(false);
+		newVirtualKeyboardInput.setBackgroundColor(Color.WHITE);
+		newVirtualKeyboardInput.setLayoutParams(new RelativeLayout.LayoutParams(LayoutParams.MATCH_PARENT, LayoutParams.WRAP_CONTENT));
+		newVirtualKeyboardInput.setVisibility(View.GONE);
+		if (ANDROID_BUILD_VERSION < 11)
+		{
+			newVirtualKeyboardInput.setImeOptions(EditorInfo.IME_FLAG_NO_EXTRACT_UI);
+		}
+		else
+		{
+			newVirtualKeyboardInput.setImeOptions(EditorInfo.IME_FLAG_NO_EXTRACT_UI | EditorInfo.IME_FLAG_NO_FULLSCREEN);
+		}
+
+		newVirtualKeyboardInput.addTextChangedListener(new TextWatcher() {
+			@Override
+			public void beforeTextChanged(CharSequence charSequence, int start, int count, int after) 
+			{
+			}
+
+			@Override
+			public void afterTextChanged(Editable s) 
+			{
+			}
+
+			@Override
+			public void onTextChanged(CharSequence charSequence, int start, int before, int count) 
+			{
+				//send to the associated Slate control
+				String message = newVirtualKeyboardInput.getText().toString();
+				nativeVirtualKeyboardChanged(message);
+			}
+		});
+		
+		virtualKeyboardLayout.addView (newVirtualKeyboardInput);
+	}
+
+	//check if the new virtual keyboard input has received a MOUSE_DOWN event
+	public boolean AndroidThunkJava_IsVirtuaInputClicked(int x, int y) 
+	{
+		View view = getCurrentFocus();
+		if (view == newVirtualKeyboardInput) 
+		{
+			Rect r = new Rect();
+			view.getGlobalVisibleRect(r);
+			if (r.contains(x, y)) 
+			{
+				return true;
+			}
+		}
+		return false;
+	}
 
 	public native boolean nativeIsShippingBuild();
 	public native void nativeSetGlobalActivity(boolean bUseExternalFilesDir, boolean bOBBInAPK, String APKPath);
@@ -2633,6 +2853,7 @@ public class GameActivity extends NativeActivity implements SurfaceHolder.Callba
 	public native void nativeGoogleClientConnectCompleted(boolean bSuccess, String accessToken);
 
 	public native void nativeVirtualKeyboardShown(int left, int top, int right, int bottom);
+	public native void nativeVirtualKeyboardVisible(boolean bShown);
 		
 	static
 	{
diff --git a/Engine/Config/Android/AndroidEngine.ini b/Engine/Config/Android/AndroidEngine.ini
--- a/Engine/Config/Android/AndroidEngine.ini
+++ b/Engine/Config/Android/AndroidEngine.ini
@@ -17,12 +17,7 @@ DeviceProfileSelectionModule="AndroidDeviceProfileSelectorRuntime"
 DefaultProviderName=AndroidAdvertising
 
 [OnlineSubsystem]
-;this is intentional, GooglePlay is a supporting OSS
-DefaultPlatformService=MCP 
-
-[OnlineSubsystemGooglePlay.Store]
-bSupportsInAppPurchasing=true
-bUseStoreV2=true
+DefaultPlatformService=GooglePlay
 
 [SlateRenderer]
 NumPreallocatedVertices=200
diff --git a/Engine/Plugins/Runtime/Analytics/Adjust/Source/AndroidAdjust/Private/AndroidAdjust.cpp b/Engine/Plugins/Runtime/Analytics/Adjust/Source/AndroidAdjust/Private/AndroidAdjust.cpp
--- a/Engine/Plugins/Runtime/Analytics/Adjust/Source/AndroidAdjust/Private/AndroidAdjust.cpp
+++ b/Engine/Plugins/Runtime/Analytics/Adjust/Source/AndroidAdjust/Private/AndroidAdjust.cpp
@@ -243,7 +243,7 @@ bool FAnalyticsProviderAdjust::StartSession(const TArray<FAnalyticsEventAttribut
 
 	for (auto Attr : Attributes)
 	{
-		AndroidThunkCpp_Adjust_AddSessionPartnerParameter(Attr.AttrName, Attr.AttrValue);
+		AndroidThunkCpp_Adjust_AddSessionPartnerParameter(Attr.AttrName, Attr.ToString());
 	}
 	RecordEvent(TEXT("SessionAttributes"), Attributes);
 
@@ -348,7 +348,7 @@ void FAnalyticsProviderAdjust::RecordEvent(const FString& EventName, const TArra
 			// add event attributes
 			for (auto Attr : Attributes)
 			{
-				AndroidThunkCpp_Adjust_Event_AddCallbackParameter(Attr.AttrName, Attr.AttrValue);
+				AndroidThunkCpp_Adjust_Event_AddCallbackParameter(Attr.AttrName, Attr.ToString());
 			}
 		}
 		AndroidThunkCpp_Adjust_SendEvent(EventToken);
@@ -452,7 +452,7 @@ void FAnalyticsProviderAdjust::RecordError(const FString& Error, const TArray<FA
 			// add event attributes
 			for (auto Attr : EventAttrs)
 			{
-				AndroidThunkCpp_Adjust_Event_AddCallbackParameter(Attr.AttrName, Attr.AttrValue);
+				AndroidThunkCpp_Adjust_Event_AddCallbackParameter(Attr.AttrName, Attr.ToString());
 			}
 		}
 
@@ -485,7 +485,7 @@ void FAnalyticsProviderAdjust::RecordProgress(const FString& ProgressType, const
 			// add event attributes
 			for (auto Attr : EventAttrs)
 			{
-				AndroidThunkCpp_Adjust_Event_AddCallbackParameter(Attr.AttrName, Attr.AttrValue);
+				AndroidThunkCpp_Adjust_Event_AddCallbackParameter(Attr.AttrName, Attr.ToString());
 			}
 		}
 
diff --git a/Engine/Plugins/Runtime/Analytics/Adjust/Source/IOSAdjust/Private/IOSAdjust.cpp b/Engine/Plugins/Runtime/Analytics/Adjust/Source/IOSAdjust/Private/IOSAdjust.cpp
--- a/Engine/Plugins/Runtime/Analytics/Adjust/Source/IOSAdjust/Private/IOSAdjust.cpp
+++ b/Engine/Plugins/Runtime/Analytics/Adjust/Source/IOSAdjust/Private/IOSAdjust.cpp
@@ -169,7 +169,7 @@ bool FAnalyticsProviderAdjust::StartSession(const TArray<FAnalyticsEventAttribut
 	for (auto Attr : Attributes)
 	{
 		NSString* IOSKey = [NSString stringWithFString : Attr.AttrName];
-		NSString* IOSValue = [NSString stringWithFString : Attr.AttrValue];
+		NSString* IOSValue = [NSString stringWithFString : Attr.ToString()];
 		[Adjust addSessionPartnerParameter:IOSKey value:IOSValue];
 	}
 	RecordEvent(TEXT("SessionAttributes"), Attributes);
@@ -276,7 +276,7 @@ void FAnalyticsProviderAdjust::RecordEvent(const FString& EventName, const TArra
 			for (auto Attr : Attributes)
 			{
 				NSString* IOSKey = [NSString stringWithFString : Attr.AttrName];
-				NSString* IOSValue = [NSString stringWithFString : Attr.AttrValue];
+				NSString* IOSValue = [NSString stringWithFString : Attr.ToString()];
 				[event addCallbackParameter:IOSKey value:IOSValue];
 			}
 		}
@@ -382,7 +382,7 @@ void FAnalyticsProviderAdjust::RecordError(const FString& Error, const TArray<FA
 			for (auto Attr : EventAttrs)
 			{
 				NSString* IOSKey = [NSString stringWithFString : Attr.AttrName];
-				NSString* IOSValue = [NSString stringWithFString : Attr.AttrValue];
+				NSString* IOSValue = [NSString stringWithFString : Attr.ToString()];
 				[event addCallbackParameter:IOSKey value:IOSValue];
 			}
 		}
@@ -417,7 +417,7 @@ void FAnalyticsProviderAdjust::RecordProgress(const FString& ProgressType, const
 			for (auto Attr : EventAttrs)
 			{
 				NSString* IOSKey = [NSString stringWithFString : Attr.AttrName];
-				NSString* IOSValue = [NSString stringWithFString : Attr.AttrValue];
+				NSString* IOSValue = [NSString stringWithFString : Attr.ToString()];
 				[event addCallbackParameter:IOSKey value:IOSValue];
 			}
 		}
diff --git a/Engine/Plugins/Runtime/Analytics/AnalyticsBlueprintLibrary/Source/AnalyticsBlueprintLibrary/Private/AnalyticsBlueprintLibrary.cpp b/Engine/Plugins/Runtime/Analytics/AnalyticsBlueprintLibrary/Source/AnalyticsBlueprintLibrary/Private/AnalyticsBlueprintLibrary.cpp
--- a/Engine/Plugins/Runtime/Analytics/AnalyticsBlueprintLibrary/Source/AnalyticsBlueprintLibrary/Private/AnalyticsBlueprintLibrary.cpp
+++ b/Engine/Plugins/Runtime/Analytics/AnalyticsBlueprintLibrary/Source/AnalyticsBlueprintLibrary/Private/AnalyticsBlueprintLibrary.cpp
@@ -13,11 +13,10 @@ DEFINE_LOG_CATEGORY_STATIC(LogAnalyticsBPLib, Display, All);
 static inline TArray<FAnalyticsEventAttribute> ConvertAttrs(const TArray<FAnalyticsEventAttr>& Attributes)
 {
 	TArray<FAnalyticsEventAttribute> Converted;
-	Converted.AddZeroed(Attributes.Num());
-	for (int32 Index = 0; Index < Attributes.Num(); Index++)
+	Converted.Reserve(Attributes.Num());
+	for (const FAnalyticsEventAttr& Attr : Attributes)
 	{
-		Converted[Index].AttrName = Attributes[Index].Name;
-		Converted[Index].AttrValue = Attributes[Index].Value;
+		Converted.Emplace(Attr.Name, Attr.Value);
 	}
 	return Converted;
 }
diff --git a/Engine/Plugins/Runtime/Analytics/FileLogging/Source/FileLogging/Private/FileLogging.cpp b/Engine/Plugins/Runtime/Analytics/FileLogging/Source/FileLogging/Private/FileLogging.cpp
--- a/Engine/Plugins/Runtime/Analytics/FileLogging/Source/FileLogging/Private/FileLogging.cpp
+++ b/Engine/Plugins/Runtime/Analytics/FileLogging/Source/FileLogging/Private/FileLogging.cpp
@@ -182,7 +182,7 @@ void FAnalyticsProviderFileLogging::RecordEvent(const FString& EventName, const
 				}
 				FileArchive->Logf(TEXT("\t\t\t{"));
 				FileArchive->Logf(TEXT("\t\t\t\t\"name\" : \"%s\","), *Attr.AttrName);
-				FileArchive->Logf(TEXT("\t\t\t\t\"value\" : \"%s\""), *Attr.AttrValue);
+				FileArchive->Logf(TEXT("\t\t\t\t\"value\" : \"%s\""), *Attr.ToString());
 				FileArchive->Logf(TEXT("\t\t\t}"));
 				bHasWrittenFirstAttr = true;
 			}
@@ -349,7 +349,7 @@ void FAnalyticsProviderFileLogging::RecordError(const FString& Error, const TArr
 			}
 			FileArchive->Logf(TEXT("\t\t\t{"));
 			FileArchive->Logf(TEXT("\t\t\t\t\"name\" : \"%s\","), *Attr.AttrName);
-			FileArchive->Logf(TEXT("\t\t\t\t\"value\" : \"%s\""), *Attr.AttrValue);
+			FileArchive->Logf(TEXT("\t\t\t\t\"value\" : \"%s\""), *Attr.ToString());
 			FileArchive->Logf(TEXT("\t\t\t}"));
 			bHasWrittenFirstAttr = true;
 		}
@@ -394,7 +394,7 @@ void FAnalyticsProviderFileLogging::RecordProgress(const FString& ProgressType,
 			}
 			FileArchive->Logf(TEXT("\t\t\t{"));
 			FileArchive->Logf(TEXT("\t\t\t\t\"name\" : \"%s\","), *Attr.AttrName);
-			FileArchive->Logf(TEXT("\t\t\t\t\"value\" : \"%s\""), *Attr.AttrValue);
+			FileArchive->Logf(TEXT("\t\t\t\t\"value\" : \"%s\""), *Attr.ToString());
 			FileArchive->Logf(TEXT("\t\t\t}"));
 			bHasWrittenFirstAttr = true;
 		}
@@ -439,7 +439,7 @@ void FAnalyticsProviderFileLogging::RecordItemPurchase(const FString& ItemId, in
 			}
 			FileArchive->Logf(TEXT("\t\t\t{"));
 			FileArchive->Logf(TEXT("\t\t\t\t\"name\" : \"%s\","), *Attr.AttrName);
-			FileArchive->Logf(TEXT("\t\t\t\t\"value\" : \"%s\""), *Attr.AttrValue);
+			FileArchive->Logf(TEXT("\t\t\t\t\"value\" : \"%s\""), *Attr.ToString());
 			FileArchive->Logf(TEXT("\t\t\t}"));
 			bHasWrittenFirstAttr = true;
 		}
@@ -484,7 +484,7 @@ void FAnalyticsProviderFileLogging::RecordCurrencyPurchase(const FString& GameCu
 			}
 			FileArchive->Logf(TEXT("\t\t\t{"));
 			FileArchive->Logf(TEXT("\t\t\t\t\"name\" : \"%s\","), *Attr.AttrName);
-			FileArchive->Logf(TEXT("\t\t\t\t\"value\" : \"%s\""), *Attr.AttrValue);
+			FileArchive->Logf(TEXT("\t\t\t\t\"value\" : \"%s\""), *Attr.ToString());
 			FileArchive->Logf(TEXT("\t\t\t}"));
 			bHasWrittenFirstAttr = true;
 		}
@@ -529,7 +529,7 @@ void FAnalyticsProviderFileLogging::RecordCurrencyGiven(const FString& GameCurre
 			}
 			FileArchive->Logf(TEXT("\t\t\t{"));
 			FileArchive->Logf(TEXT("\t\t\t\t\"name\" : \"%s\","), *Attr.AttrName);
-			FileArchive->Logf(TEXT("\t\t\t\t\"value\" : \"%s\""), *Attr.AttrValue);
+			FileArchive->Logf(TEXT("\t\t\t\t\"value\" : \"%s\""), *Attr.ToString());
 			FileArchive->Logf(TEXT("\t\t\t}"));
 			bHasWrittenFirstAttr = true;
 		}
diff --git a/Engine/Plugins/Runtime/Analytics/Flurry/Source/IOSFlurry/Source/IOSFlurry/Private/IOSFlurry.cpp b/Engine/Plugins/Runtime/Analytics/Flurry/Source/IOSFlurry/Source/IOSFlurry/Private/IOSFlurry.cpp
--- a/Engine/Plugins/Runtime/Analytics/Flurry/Source/IOSFlurry/Source/IOSFlurry/Private/IOSFlurry.cpp
+++ b/Engine/Plugins/Runtime/Analytics/Flurry/Source/IOSFlurry/Source/IOSFlurry/Private/IOSFlurry.cpp
@@ -211,7 +211,7 @@ void FAnalyticsProviderFlurry::RecordEvent(const FString& EventName, const TArra
 			for	(auto Attr : Attributes)
 			{
 				NSString* AttrName = [NSString stringWithFString : Attr.AttrName];
-				NSString* AttrValue = [NSString stringWithFString : Attr.AttrValue];
+				NSString* AttrValue = [NSString stringWithFString : Attr.ToString()];
 				[AttributesDict setValue:AttrValue forKey:AttrName];
 			}
 			[Flurry logEvent:ConvertedEventName withParameters:AttributesDict];
@@ -296,7 +296,7 @@ void FAnalyticsProviderFlurry::RecordItemPurchase(const FString& ItemId, int Ite
 	for	(auto Attr : EventAttrs)
 	{
 		NSString* AttrName = [NSString stringWithFString : Attr.AttrName];
-		NSString* AttrValue = [NSString stringWithFString : Attr.AttrValue];
+		NSString* AttrValue = [NSString stringWithFString : Attr.ToString()];
 		[AttributesDict setValue:AttrValue forKey:AttrName];
 	}
 	// Send the event
@@ -320,7 +320,7 @@ void FAnalyticsProviderFlurry::RecordCurrencyPurchase(const FString& GameCurrenc
 	for	(auto Attr : EventAttrs)
 	{
 		NSString* AttrName = [NSString stringWithFString : Attr.AttrName];
-		NSString* AttrValue = [NSString stringWithFString : Attr.AttrValue];
+		NSString* AttrValue = [NSString stringWithFString : Attr.ToString()];
 		[AttributesDict setValue:AttrValue forKey:AttrName];
 	}
 	// Send the event
@@ -344,7 +344,7 @@ void FAnalyticsProviderFlurry::RecordCurrencyGiven(const FString& GameCurrencyTy
 	for	(auto Attr : EventAttrs)
 	{
 		NSString* AttrName = [NSString stringWithFString : Attr.AttrName];
-		NSString* AttrValue = [NSString stringWithFString : Attr.AttrValue];
+		NSString* AttrValue = [NSString stringWithFString : Attr.ToString()];
 		[AttributesDict setValue:AttrValue forKey:AttrName];
 	}
 	// Send the event
@@ -367,7 +367,7 @@ void FAnalyticsProviderFlurry::RecordError(const FString& Error, const TArray<FA
 	for	(auto Attr : EventAttrs)
 	{
 		NSString* AttrName = [NSString stringWithFString : Attr.AttrName];
-		NSString* AttrValue = [NSString stringWithFString : Attr.AttrValue];
+		NSString* AttrValue = [NSString stringWithFString : Attr.ToString()];
 		[AttributesDict setValue:AttrValue forKey:AttrName];
 	}
 	// Send the event
@@ -391,7 +391,7 @@ void FAnalyticsProviderFlurry::RecordProgress(const FString& ProgressType, const
 	for	(auto Attr : EventAttrs)
 	{
 		NSString* AttrName = [NSString stringWithFString : Attr.AttrName];
-		NSString* AttrValue = [NSString stringWithFString : Attr.AttrValue];
+		NSString* AttrValue = [NSString stringWithFString : Attr.ToString()];
 		[AttributesDict setValue:AttrValue forKey:AttrName];
 	}
 	// Send the event
diff --git a/Engine/Source/Editor/ContentBrowser/Private/FrontendFilters.cpp b/Engine/Source/Editor/ContentBrowser/Private/FrontendFilters.cpp
--- a/Engine/Source/Editor/ContentBrowser/Private/FrontendFilters.cpp
+++ b/Engine/Source/Editor/ContentBrowser/Private/FrontendFilters.cpp
@@ -96,8 +96,7 @@ public:
 			}
 		}
 
-		check(AliasToSourceTagMapping.IsValid());
-		return AliasToSourceTagMapping->FindRef(InAlias);
+		return AliasToSourceTagMapping.IsValid() ? AliasToSourceTagMapping->FindRef(InAlias) : NAME_None;
 	}
 
 private:
diff --git a/Engine/Source/Editor/MaterialEditor/Private/MaterialEditorDetailCustomization.cpp b/Engine/Source/Editor/MaterialEditor/Private/MaterialEditorDetailCustomization.cpp
--- a/Engine/Source/Editor/MaterialEditor/Private/MaterialEditorDetailCustomization.cpp
+++ b/Engine/Source/Editor/MaterialEditor/Private/MaterialEditorDetailCustomization.cpp
@@ -405,7 +405,6 @@ void FMaterialDetailCustomization::CustomizeDetails( IDetailLayoutBuilder& Detai
 		DetailLayout.HideCategory( TEXT("TranslucencySelfShadowing") );
 		DetailLayout.HideCategory( TEXT("Translucency") );
 		DetailLayout.HideCategory( TEXT("Tessellation") );
-		DetailLayout.HideCategory( TEXT("Mobile") );
 		DetailLayout.HideCategory( TEXT("PostProcessMaterial") );
 		DetailLayout.HideCategory( TEXT("Lightmass") );
 		DetailLayout.HideCategory( TEXT("Thumbnail") );
@@ -413,22 +412,44 @@ void FMaterialDetailCustomization::CustomizeDetails( IDetailLayoutBuilder& Detai
 		DetailLayout.HideCategory( TEXT("PhysicalMaterial") );
 		DetailLayout.HideCategory( TEXT("Usage") );
 
-		IDetailCategoryBuilder& MaterialCategory = DetailLayout.EditCategory( TEXT("Material") );
+		// Materail category
+		{
+			IDetailCategoryBuilder& MaterialCategory = DetailLayout.EditCategory( TEXT("Material") );
+
+			TArray<TSharedRef<IPropertyHandle>> AllProperties;
+			MaterialCategory.GetDefaultProperties( AllProperties );
 
-		TArray<TSharedRef<IPropertyHandle>> AllProperties;
-		MaterialCategory.GetDefaultProperties( AllProperties );
+			for( TSharedRef<IPropertyHandle>& PropertyHandle : AllProperties )
+			{
+				UProperty* Property = PropertyHandle->GetProperty();
+				FName PropertyName = Property->GetFName();
+
+				if(		PropertyName != GET_MEMBER_NAME_CHECKED(UMaterial, MaterialDomain) 
+					&&	PropertyName != GET_MEMBER_NAME_CHECKED(UMaterial, BlendMode) 
+					&&	PropertyName != GET_MEMBER_NAME_CHECKED(UMaterial, OpacityMaskClipValue) 
+					&&  	PropertyName != GET_MEMBER_NAME_CHECKED(UMaterial, NumCustomizedUVs) )
+				{
+					DetailLayout.HideProperty( PropertyHandle );
+				}
+			}
+		}
 
-		for( TSharedRef<IPropertyHandle>& PropertyHandle : AllProperties )
+		// Mobile category
 		{
-			UProperty* Property = PropertyHandle->GetProperty();
-			FName PropertyName = Property->GetFName();
+			IDetailCategoryBuilder& MobileCategory = DetailLayout.EditCategory(TEXT("Mobile"));
+			
+			TArray<TSharedRef<IPropertyHandle>> AllProperties;
+			MobileCategory.GetDefaultProperties(AllProperties);
 
-			if(		PropertyName != GET_MEMBER_NAME_CHECKED(UMaterial, MaterialDomain) 
-				&&	PropertyName != GET_MEMBER_NAME_CHECKED(UMaterial, BlendMode) 
-				&&	PropertyName != GET_MEMBER_NAME_CHECKED(UMaterial, OpacityMaskClipValue) 
-				&&  	PropertyName != GET_MEMBER_NAME_CHECKED(UMaterial, NumCustomizedUVs) )
+			for (TSharedRef<IPropertyHandle>& PropertyHandle : AllProperties)
 			{
-				DetailLayout.HideProperty( PropertyHandle );
+				UProperty* Property = PropertyHandle->GetProperty();
+				FName PropertyName = Property->GetFName();
+
+				if (PropertyName != GET_MEMBER_NAME_CHECKED(UMaterial, bUseFullPrecision)) 
+				{
+					DetailLayout.HideProperty(PropertyHandle);
+				}
 			}
 		}
 	}
diff --git a/Engine/Source/Editor/PropertyEditor/Private/UserInterface/PropertyEditor/SPropertyEditorArrayItem.cpp b/Engine/Source/Editor/PropertyEditor/Private/UserInterface/PropertyEditor/SPropertyEditorArrayItem.cpp
--- a/Engine/Source/Editor/PropertyEditor/Private/UserInterface/PropertyEditor/SPropertyEditorArrayItem.cpp
+++ b/Engine/Source/Editor/PropertyEditor/Private/UserInterface/PropertyEditor/SPropertyEditorArrayItem.cpp
@@ -50,10 +50,21 @@ bool SPropertyEditorArrayItem::Supports( const TSharedRef< class FPropertyEditor
 	const TSharedRef< FPropertyNode > PropertyNode = PropertyEditor->GetPropertyNode();
 	const UProperty* Property = PropertyEditor->GetProperty();
 
-	return !Cast<const UClassProperty>( Property ) &&
-		Cast<const UArrayProperty>( Property->GetOuter() ) &&
-		PropertyNode->HasNodeFlags(EPropertyNodeFlags::SingleSelectOnly) &&
-		!(Cast<const UArrayProperty>(Property->GetOuter())->PropertyFlags & CPF_EditConst);
+	if (!Cast<const UClassProperty>(Property) && PropertyNode->HasNodeFlags(EPropertyNodeFlags::SingleSelectOnly))
+	{
+		if (Cast<const UArrayProperty>(Property->GetOuter()) &&
+			!(Cast<const UArrayProperty>(Property->GetOuter())->PropertyFlags & CPF_EditConst))
+		{
+			return true;
+		}
+
+		if (Cast<const UMapProperty>(Property->GetOuter()) &&
+			!(Cast<const UMapProperty>(Property->GetOuter())->PropertyFlags & CPF_EditConst))
+		{
+			return true;
+		}
+	}
+	return false;
 }
 
 FText SPropertyEditorArrayItem::GetValueAsString() const
diff --git a/Engine/Source/Programs/AutomationTool/IOS/IOSPlatform.Automation.cs b/Engine/Source/Programs/AutomationTool/IOS/IOSPlatform.Automation.cs
--- a/Engine/Source/Programs/AutomationTool/IOS/IOSPlatform.Automation.cs
+++ b/Engine/Source/Programs/AutomationTool/IOS/IOSPlatform.Automation.cs
@@ -1533,7 +1533,7 @@ public class IOSPlatform : Platform
 
 	public override void PostStagingFileCopy(ProjectParams Params, DeploymentContext SC)
 	{
-		if (Params.CreateChunkInstall)
+/*		if (Params.CreateChunkInstall)
 		{
 			// get the bundle identifier
 			string BundleIdentifier = "";
@@ -1605,7 +1605,7 @@ public class IOSPlatform : Platform
 
 			// generate the OnDemandResources.plist
 			GenerateOnDemandResourcesPlist (ChunkData, SC.StageDirectory.FullName);
-		}
+		}*/
 	}
 
 	public override bool StageMovies
diff --git a/Engine/Source/Programs/AutomationTool/Scripts/CopyBuildToStagingDirectory.Automation.cs b/Engine/Source/Programs/AutomationTool/Scripts/CopyBuildToStagingDirectory.Automation.cs
--- a/Engine/Source/Programs/AutomationTool/Scripts/CopyBuildToStagingDirectory.Automation.cs
+++ b/Engine/Source/Programs/AutomationTool/Scripts/CopyBuildToStagingDirectory.Automation.cs
@@ -1261,10 +1261,10 @@ public partial class Project : CommandUtils
 			bool bAddedToChunk = false;
 			for (int ChunkIndex = 0; !bAddedToChunk && ChunkIndex < ChunkResponseFiles.Length; ++ChunkIndex)
 			{
-                		string OriginalFilename = StagingFile.Key;
-                		string NoExtension = CombinePaths(Path.GetDirectoryName(OriginalFilename), Path.GetFileNameWithoutExtension(OriginalFilename));
-                		string OriginalReplaceSlashes = OriginalFilename.Replace('/', '\\');
-                		string NoExtensionReplaceSlashes = NoExtension.Replace('/', '\\');
+				string OriginalFilename = StagingFile.Key;
+				string NoExtension = CombinePaths(Path.GetDirectoryName(OriginalFilename), Path.GetFileNameWithoutExtension(OriginalFilename));
+				string OriginalReplaceSlashes = OriginalFilename.Replace('/', '\\');
+				string NoExtensionReplaceSlashes = NoExtension.Replace('/', '\\');
 
 				if (ChunkResponseFiles[ChunkIndex].Contains(OriginalFilename) || 
                 		    ChunkResponseFiles[ChunkIndex].Contains(OriginalReplaceSlashes) ||
@@ -1302,7 +1302,7 @@ public partial class Project : CommandUtils
 		IEnumerable<Tuple<Dictionary<string,string>, string>> PakPairs = PakResponseFiles.Zip(ChunkList, (a, b) => Tuple.Create(a, b));
 	
         System.Threading.Tasks.ParallelOptions Options = new System.Threading.Tasks.ParallelOptions();
-        Options.MaxDegreeOfParallelism = 1;
+        Options.MaxDegreeOfParallelism = 16;
 		System.Threading.Tasks.Parallel.ForEach(PakPairs, (PakPair) =>
 		{
 			var ChunkName = Path.GetFileNameWithoutExtension(PakPair.Item2);
diff --git a/Engine/Source/Programs/UnrealBuildTool/Android/AndroidToolChain.cs b/Engine/Source/Programs/UnrealBuildTool/Android/AndroidToolChain.cs
--- a/Engine/Source/Programs/UnrealBuildTool/Android/AndroidToolChain.cs
+++ b/Engine/Source/Programs/UnrealBuildTool/Android/AndroidToolChain.cs
@@ -1435,6 +1435,9 @@ namespace UnrealBuildTool
 
 		public override void ModifyBuildProducts(ReadOnlyTargetRules Target, UEBuildBinary Binary, List<string> Libraries, List<UEBuildBundleResource> BundleResources, Dictionary<FileReference, BuildProductType> BuildProducts)
 		{
+			// only the .so needs to be in the manifest; we always have to build the apk since its contents depend on the project
+
+			/*
 			// the binary will have all of the .so's in the output files, we need to trim down to the shared apk (which is what needs to go into the manifest)
 			if (Target.bDeployAfterCompile && Binary.Config.Type != UEBuildBinaryType.StaticLibrary)
 			{
@@ -1444,6 +1447,7 @@ namespace UnrealBuildTool
 					BuildProducts.Add(ApkFile, BuildProductType.Package);
 				}
 			}
+			*/
 		}
 
 		public override void CompileCSharpProject(CSharpEnvironment CompileEnvironment, FileReference ProjectFileName, FileReference DestinationFile, ActionGraph ActionGraph)
diff --git a/Engine/Source/Programs/UnrealBuildTool/Configuration/ModuleRules.cs b/Engine/Source/Programs/UnrealBuildTool/Configuration/ModuleRules.cs
--- a/Engine/Source/Programs/UnrealBuildTool/Configuration/ModuleRules.cs
+++ b/Engine/Source/Programs/UnrealBuildTool/Configuration/ModuleRules.cs
@@ -230,6 +230,11 @@ namespace UnrealBuildTool
 		/// </summary>
 		public bool bEnableExceptions = false;
 
+		/// <summary>
+		/// Enable objective C exception handling
+		/// </summary>
+		public bool bEnableObjCExceptions = false;
+
 		/// <summary>
 		/// Enable warnings for shadowed variables
 		/// </summary>
diff --git a/Engine/Source/Programs/UnrealBuildTool/Configuration/TargetRules.cs b/Engine/Source/Programs/UnrealBuildTool/Configuration/TargetRules.cs
--- a/Engine/Source/Programs/UnrealBuildTool/Configuration/TargetRules.cs
+++ b/Engine/Source/Programs/UnrealBuildTool/Configuration/TargetRules.cs
@@ -462,6 +462,11 @@ namespace UnrealBuildTool
 		/// </summary>
 		public bool bForceEnableExceptions = false;
 
+		/// <summary>
+		/// Enable exceptions for all modules.
+		/// </summary>
+		public bool bForceEnableObjCExceptions = false;
+
 		/// <summary>
 		/// Enable RTTI for all modules.
 		/// </summary>
@@ -1622,6 +1627,11 @@ namespace UnrealBuildTool
 			get { return Inner.bForceEnableExceptions; }
 		}
 
+		public bool bForceEnableObjCExceptions
+		{
+			get { return Inner.bForceEnableObjCExceptions; }
+		}
+
 		public bool bForceEnableRTTI
 		{
 			get { return Inner.bForceEnableRTTI; }
diff --git a/Engine/Source/Programs/UnrealBuildTool/Configuration/UEBuildModuleCPP.cs b/Engine/Source/Programs/UnrealBuildTool/Configuration/UEBuildModuleCPP.cs
--- a/Engine/Source/Programs/UnrealBuildTool/Configuration/UEBuildModuleCPP.cs
+++ b/Engine/Source/Programs/UnrealBuildTool/Configuration/UEBuildModuleCPP.cs
@@ -1075,6 +1075,7 @@ namespace UnrealBuildTool
 			Result.MinFilesUsingPrecompiledHeaderOverride = Rules.MinFilesUsingPrecompiledHeaderOverride;
 			Result.bBuildLocallyWithSNDBS = Rules.bBuildLocallyWithSNDBS;
 			Result.bEnableExceptions |= Rules.bEnableExceptions;
+			Result.bEnableObjCExceptions |= Rules.bEnableObjCExceptions;
 			Result.bEnableShadowVariableWarnings = Rules.bEnableShadowVariableWarnings;
 			Result.bEnableUndefinedIdentifierWarnings = Rules.bEnableUndefinedIdentifierWarnings;
 			Result.bUseStaticCRT = Target.bUseStaticCRT;
diff --git a/Engine/Source/Programs/UnrealBuildTool/Configuration/UEBuildTarget.cs b/Engine/Source/Programs/UnrealBuildTool/Configuration/UEBuildTarget.cs
--- a/Engine/Source/Programs/UnrealBuildTool/Configuration/UEBuildTarget.cs
+++ b/Engine/Source/Programs/UnrealBuildTool/Configuration/UEBuildTarget.cs
@@ -4149,6 +4149,7 @@ namespace UnrealBuildTool
 			GlobalCompileEnvironment.Definitions.Add(String.Format("IS_PROGRAM={0}", TargetType == TargetType.Program ? "1" : "0"));
 			GlobalCompileEnvironment.Definitions.AddRange(Rules.GlobalDefinitions);
 			GlobalCompileEnvironment.bEnableExceptions = Rules.bForceEnableExceptions || Rules.bBuildEditor;
+			GlobalCompileEnvironment.bEnableObjCExceptions = Rules.bForceEnableObjCExceptions || Rules.bBuildEditor;
 			GlobalCompileEnvironment.bShadowVariableWarningsAsErrors = Rules.bShadowVariableErrors;
 			GlobalCompileEnvironment.bUndefinedIdentifierWarningsAsErrors = Rules.bUndefinedIdentifierErrors;
 			GlobalCompileEnvironment.bOptimizeForSize = Rules.bCompileForSize;
diff --git a/Engine/Source/Programs/UnrealBuildTool/IOS/IOSToolChain.cs b/Engine/Source/Programs/UnrealBuildTool/IOS/IOSToolChain.cs
--- a/Engine/Source/Programs/UnrealBuildTool/IOS/IOSToolChain.cs
+++ b/Engine/Source/Programs/UnrealBuildTool/IOS/IOSToolChain.cs
@@ -173,7 +173,16 @@ namespace UnrealBuildTool
 			{
 				Result += " -fno-exceptions";
 			}
-			
+
+			if (CompileEnvironment.bEnableObjCExceptions)
+			{
+				Result += " -fobjc-exceptions";
+			}
+			else
+			{
+				Result += " -fno-objc-exceptions";
+			}
+
 			string SanitizerMode = Environment.GetEnvironmentVariable("ENABLE_ADDRESS_SANITIZER");
 			if(SanitizerMode != null && SanitizerMode == "YES")
 			{
@@ -1324,6 +1333,32 @@ namespace UnrealBuildTool
 			}
 		}
 
+		private static void GenerateCrashlyticsData(string ExecutableDirectory, string ExecutableName, string ProjectDir, string ProjectName)
+        {
+			Log.TraceInformation("Generating and uploading Crashlytics Data");
+            string FabricPath = UnrealBuildTool.EngineDirectory + "/Intermediate/UnzippedFrameworks/ThirdPartyFrameworks/Fabric.embeddedframework";
+            if (Directory.Exists(FabricPath))
+            {
+				string PlistFile = ProjectDir + "/Intermediate/IOS/" + ProjectName + "-Info.plist";
+                Process FabricProcess = new Process();
+                FabricProcess.StartInfo.WorkingDirectory = ExecutableDirectory;
+                FabricProcess.StartInfo.FileName = "/bin/sh";
+				FabricProcess.StartInfo.Arguments = string.Format("-c 'chmod 777 \"{0}/Fabric.framework/upload-symbols\"; \"{0}/Fabric.framework/upload-symbols\" -a 7a4cebd0324af21696e5e321802c5e26ba541cad -p ios {1}'",
+                    FabricPath,
+					ExecutableDirectory + "/" + ExecutableName + ".dSYM.zip");
+                FabricProcess.OutputDataReceived += new DataReceivedEventHandler(OutputReceivedDataEventHandler);
+                FabricProcess.ErrorDataReceived += new DataReceivedEventHandler(OutputReceivedDataEventHandler);
+
+                OutputReceivedDataEventHandlerEncounteredError = false;
+                OutputReceivedDataEventHandlerEncounteredErrorMessage = "";
+                Utils.RunLocalProcess(FabricProcess);
+                if (OutputReceivedDataEventHandlerEncounteredError)
+                {
+                    throw new Exception(OutputReceivedDataEventHandlerEncounteredErrorMessage);
+                }
+            }
+        }
+
         public static void PostBuildSync(UEBuildTarget Target)
 		{
 			IOSProjectSettings ProjectSettings = ((IOSPlatform)UEBuildPlatform.GetBuildPlatform(UnrealTargetPlatform.IOS)).ReadProjectSettings(Target.ProjectFile);
@@ -1354,6 +1389,10 @@ namespace UnrealBuildTool
 					}
 				}
 
+                // ensure the plist, entitlements, and provision files are properly copied
+                var DeployHandler = (Target.Platform == UnrealTargetPlatform.IOS ? new UEDeployIOS() : new UEDeployTVOS());
+                DeployHandler.PrepTargetForDeployment(new UEBuildDeployTarget(Target));
+
 				// copy the executable
 				if (!File.Exists(FinalRemoteExecutablePath))
 				{
@@ -1361,16 +1400,14 @@ namespace UnrealBuildTool
 				}
 				File.Copy(Target.OutputPath.FullName, FinalRemoteExecutablePath, true);
 
+				GenerateCrashlyticsData(RemoteShadowDirectoryMac, Path.GetFileName(Target.OutputPath.FullName), Target.ProjectDirectory.FullName, AppName);
+
 				if (Target.Rules.bCreateStubIPA)
 				{
 					string Project = Target.ProjectDirectory + "/" + AppName + ".uproject";
 
 					string SchemeName = AppName;
 
-                    // ensure the plist, entitlements, and provision files are properly copied
-                    var DeployHandler = (Target.Platform == UnrealTargetPlatform.IOS ? new UEDeployIOS() : new UEDeployTVOS());
-                    DeployHandler.PrepTargetForDeployment(new UEBuildDeployTarget(Target));
-
                     // generate the dummy project so signing works
                     if (AppName == "UE4Game" || AppName == "UE4Client" || Utils.IsFileUnderDirectory(Target.ProjectDirectory + "/" + AppName + ".uproject", Path.GetFullPath("../..")))
 					{
diff --git a/Engine/Source/Programs/UnrealBuildTool/IOS/UEBuildIOS.cs b/Engine/Source/Programs/UnrealBuildTool/IOS/UEBuildIOS.cs
--- a/Engine/Source/Programs/UnrealBuildTool/IOS/UEBuildIOS.cs
+++ b/Engine/Source/Programs/UnrealBuildTool/IOS/UEBuildIOS.cs
@@ -141,6 +141,12 @@ namespace UnrealBuildTool
 		[ConfigFile(ConfigHierarchyType.Engine, "/Script/IOSRuntimeSettings.IOSRuntimeSettings", "bAutomaticSigning")]
 		public readonly bool bAutomaticSigning = false;
 
+		/// <summary>
+		/// true to change FORCEINLINE to a regular INLINE.
+		/// </summary>
+		[ConfigFile(ConfigHierarchyType.Engine, "/Script/IOSRuntimeSettings.IOSRuntimeSettings", "bDisableForceInline")]
+		public readonly bool bDisableForceInline = false;
+		
 		/// <summary>
 		/// Returns a list of all the non-shipping architectures which are supported
 		/// </summary>
@@ -752,6 +758,8 @@ namespace UnrealBuildTool
 				CompileEnvironment.Definitions.Add("NOTIFICATIONS_ENABLED=0");
 			}
 
+			CompileEnvironment.Definitions.Add("UE_DISABLE_FORCE_INLINE=" + (ProjectSettings.bDisableForceInline ? "1" : "0"));
+
 			if (Target.Architecture == "-simulator")
 			{
 				CompileEnvironment.Definitions.Add("WITH_SIMULATOR=1");
diff --git a/Engine/Source/Programs/UnrealBuildTool/IOS/UEDeployIOS.cs b/Engine/Source/Programs/UnrealBuildTool/IOS/UEDeployIOS.cs
--- a/Engine/Source/Programs/UnrealBuildTool/IOS/UEDeployIOS.cs
+++ b/Engine/Source/Programs/UnrealBuildTool/IOS/UEDeployIOS.cs
@@ -320,6 +320,15 @@ namespace UnrealBuildTool
 			string FacebookDisplayName = "";
 			Ini.GetString("/Script/IOSRuntimeSettings.IOSRuntimeSettings", "BundleDisplayName", out FacebookDisplayName);
 
+			// Get Google Support details
+			bool bEnableGoogleSupport = true;
+			Ini.GetBool("/Script/IOSRuntimeSettings.IOSRuntimeSettings", "bEnableGoogleSupport", out bEnableGoogleSupport);
+
+			// Write the Google App ID if we need it.
+			string GoogleReversedClientId = "";
+			Ini.GetString("/Script/IOSRuntimeSettings.IOSRuntimeSettings", "GoogleReversedClientId", out GoogleReversedClientId);
+			bEnableGoogleSupport = bEnableFacebookSupport && !string.IsNullOrWhiteSpace(GoogleReversedClientId);
+
 			// Add remote-notifications as background mode
 			bool bRemoteNotificationsSupported = false;
 			Ini.GetBool("/Script/IOSRuntimeSettings.IOSRuntimeSettings", "bEnableRemoteNotificationsSupport", out bRemoteNotificationsSupported);
@@ -353,6 +362,10 @@ namespace UnrealBuildTool
 				// This is needed for facebook login to redirect back to the app after completion.
 				Text.AppendLine(string.Format("\t\t\t\t<string>fb{0}</string>", FacebookAppID));
 			}
+			if (bEnableGoogleSupport)
+			{
+				Text.AppendLine(string.Format("\t\t\t\t<string>{0}</string>", GoogleReversedClientId));
+			}
 			Text.AppendLine("\t\t\t</array>");
 			Text.AppendLine("\t\t</dict>");
 			Text.AppendLine("\t</array>");
diff --git a/Engine/Source/Programs/UnrealBuildTool/System/CppCompileEnvironment.cs b/Engine/Source/Programs/UnrealBuildTool/System/CppCompileEnvironment.cs
--- a/Engine/Source/Programs/UnrealBuildTool/System/CppCompileEnvironment.cs
+++ b/Engine/Source/Programs/UnrealBuildTool/System/CppCompileEnvironment.cs
@@ -169,6 +169,11 @@ namespace UnrealBuildTool
 		/// </summary>
 		public bool bEnableExceptions = false;
 
+		/// <summary>
+		/// Enable objective C exception handling
+		/// </summary>
+		public bool bEnableObjCExceptions = false;
+
 		/// <summary>
 		/// Whether to warn about the use of shadow variables
 		/// </summary>
@@ -336,6 +341,7 @@ namespace UnrealBuildTool
 			MinFilesUsingPrecompiledHeaderOverride = Other.MinFilesUsingPrecompiledHeaderOverride;
 			bBuildLocallyWithSNDBS = Other.bBuildLocallyWithSNDBS;
 			bEnableExceptions = Other.bEnableExceptions;
+			bEnableObjCExceptions = Other.bEnableObjCExceptions;
 			bShadowVariableWarningsAsErrors = Other.bShadowVariableWarningsAsErrors;
 			bEnableShadowVariableWarnings = Other.bEnableShadowVariableWarnings;
 			bUndefinedIdentifierWarningsAsErrors = Other.bUndefinedIdentifierWarningsAsErrors;
diff --git a/Engine/Source/Programs/UnrealWatchdog/Private/UnrealWatchdog.cpp b/Engine/Source/Programs/UnrealWatchdog/Private/UnrealWatchdog.cpp
--- a/Engine/Source/Programs/UnrealWatchdog/Private/UnrealWatchdog.cpp
+++ b/Engine/Source/Programs/UnrealWatchdog/Private/UnrealWatchdog.cpp
@@ -144,7 +144,8 @@ void SendHeartbeatEvent(IAnalyticsProviderET& Analytics, const FWatchdogCommandL
 		// Update the IntervalSec with the exact time since the last heartbeat
 		FDateTime EventTime = FDateTime::UtcNow();
 		FTimespan ActualInterval = EventTime - LastEventTime;
-		HeartbeatAttributes[HeartbeatAttributes.Num() - 1] = FAnalyticsEventAttribute(TEXT("IntervalSec"), ActualInterval.GetTotalSeconds());
+		HeartbeatAttributes.RemoveAt(HeartbeatAttributes.Num() - 1);
+		HeartbeatAttributes.Add(FAnalyticsEventAttribute(TEXT("IntervalSec"), ActualInterval.GetTotalSeconds()));
 		LastEventTime = EventTime;
 	}
 
@@ -186,8 +187,8 @@ void SendHangDetectedEvent(IAnalyticsProviderET& Analytics, const FWatchdogComma
 	GetCommonEventAttributes(CommandLine, HangAttributes);
 
 	// Internal builds should popup dialogs for hangs
-	FAnalyticsEventAttribute HangUserResponse(TEXT("HangUserResponse"), TEXT("Unattended"));
-	FAnalyticsEventAttribute RecoveredUserResponse(TEXT("AlreadyRecoveredUserResponse"), TEXT("Unattended"));
+	FString HangResponse = TEXT("Unattended");
+	FString RecoveredResponse = TEXT("Unattended");
 	if (CommandLine.bAllowDialogs)
 	{
 		FText SessionLabel = FText::Format(LOCTEXT("HangSessionLabel", "{0} ({1})"), FText::FromString(CommandLine.ProjectName), FText::FromString(CommandLine.RunType));
@@ -200,7 +201,7 @@ void SendHangDetectedEvent(IAnalyticsProviderET& Analytics, const FWatchdogComma
 		if (HangAnswer == EAppReturnType::Yes)
 		{
 			UE_LOG(UnrealWatchdogLog, Log, TEXT("User confirmed hang"));
-			HangUserResponse.AttrValue = TEXT("Confirmed");
+			HangResponse= TEXT("Confirmed");
 
 			FText RecoveredMessage(LOCTEXT("WatchdogPopupRecoveredMessage", "Has the application recovered?"));
 
@@ -209,12 +210,12 @@ void SendHangDetectedEvent(IAnalyticsProviderET& Analytics, const FWatchdogComma
 			if (RecoveredAnswer == EAppReturnType::Yes)
 			{
 				UE_LOG(UnrealWatchdogLog, Log, TEXT("User confirmed recovery from hang"));
-				RecoveredUserResponse.AttrValue = TEXT("Recovered");
+				RecoveredResponse = TEXT("Recovered");
 			}
 			else
 			{
 				UE_LOG(UnrealWatchdogLog, Log, TEXT("User confirmed hang not yet recovered"));
-				RecoveredUserResponse.AttrValue = TEXT("NotRecovered");
+				RecoveredResponse = TEXT("NotRecovered");
 			}
 
 			FText CRCMessage(LOCTEXT("WatchdogPopupHangCRCMessage", "We will now open the Crash Reporter for you to tell us what happened."));
@@ -226,13 +227,13 @@ void SendHangDetectedEvent(IAnalyticsProviderET& Analytics, const FWatchdogComma
 		else  // HangAnswer == EAppReturnType::No
 		{
 			UE_LOG(UnrealWatchdogLog, Warning, TEXT("User didn't witness hang. False positive warning!"));
-			HangUserResponse.AttrValue = TEXT("False");
-			RecoveredUserResponse.AttrValue = TEXT("N/A");
+			HangResponse = TEXT("False");
+			RecoveredResponse = TEXT("N/A");
 		}
 	}
 
-	HangAttributes.Add(HangUserResponse);
-	HangAttributes.Add(RecoveredUserResponse);
+	HangAttributes.Add(FAnalyticsEventAttribute(TEXT("HangUserResponse"), HangResponse));
+	HangAttributes.Add(FAnalyticsEventAttribute(TEXT("AlreadyRecoveredUserResponse"), RecoveredResponse));
 	UE_LOG(UnrealWatchdogLog, Log, TEXT("Sending event UnrealWatchdog.HangDetected"));
 	Analytics.RecordEvent(TEXT("UnrealWatchdog.HangDetected"), MoveTemp(HangAttributes));
 }
@@ -452,7 +453,7 @@ int RunUnrealWatchdog(const TCHAR* CommandLine)
 	}
 
 	// Optional section for dialogs and CRC in internal builds
-	FAnalyticsEventAttribute UserResponse(TEXT("AbnormalShutdownUserResponse"), TEXT("Unattended"));
+	FString ShutdownResponse = TEXT("Unattended");
 	if (WatchdogCommandLine.bAllowDialogs && !bHang && StoredValues.WasDebugged != WatchdogDefs::TrueValueString)
 	{
 		EAppReturnType::Type UserAnswer = EAppReturnType::Cancel;
@@ -472,12 +473,12 @@ int RunUnrealWatchdog(const TCHAR* CommandLine)
 			if (UserAnswer == EAppReturnType::Yes)
 			{
 				UE_LOG(UnrealWatchdogLog, Log, TEXT("User confirmed crash and crash report client"));
-				UserResponse.AttrValue = TEXT("Confirmed");
+				ShutdownResponse = TEXT("Confirmed");
 			}
 			else if (UserAnswer == EAppReturnType::No)
 			{
 				UE_LOG(UnrealWatchdogLog, Warning, TEXT("User didn't witness crash and crash report client. False positive warning!"));
-				UserResponse.AttrValue = TEXT("False");
+				ShutdownResponse = TEXT("False");
 
 				FText CRCMessage(LOCTEXT("WatchdogPopupCRCMessage", "We will now open the Crash Reporter for you to tell us what happened."));
 				FMessageDialog::Open(EAppMsgType::Ok, CRCMessage, &MessageTitle);
@@ -551,7 +552,7 @@ int RunUnrealWatchdog(const TCHAR* CommandLine)
 			if (UserAnswer == EAppReturnType::Yes)
 			{
 				UE_LOG(UnrealWatchdogLog, Log, TEXT("User confirmed abnormal shutdown"));
-				UserResponse.AttrValue = TEXT("Confirmed");
+				ShutdownResponse = TEXT("Confirmed");
 
 				FText MessageTitleFormat(LOCTEXT("WatchdogPopupTitleAbnormalShutdown", "{0} terminated unexpectedly"));
 				FText MessageTitle = FText::Format(MessageTitleFormat, SessionLabel);
@@ -564,7 +565,7 @@ int RunUnrealWatchdog(const TCHAR* CommandLine)
 			else if (UserAnswer == EAppReturnType::No)
 			{
 				UE_LOG(UnrealWatchdogLog, Warning, TEXT("User didn't witness abnormal shutdown. False positive warning!"));
-				UserResponse.AttrValue = TEXT("False");
+				ShutdownResponse = TEXT("False");
 			}
 		}
 	}
@@ -572,6 +573,7 @@ int RunUnrealWatchdog(const TCHAR* CommandLine)
 	// Send watchdog shutdown event
 	UE_LOG(UnrealWatchdogLog, Log, TEXT("Watchdog watched process exited. bReturnCodeObtained=%s, ReturnCode=%u, RecordedShutdownType=%s"),
 		bReturnCodeObtained ? TEXT("1") : TEXT("0"), ReturnCode, *StoredValues.ExecutionStatus);
+	FAnalyticsEventAttribute UserResponse(TEXT("AbnormalShutdownUserResponse"), ShutdownResponse);
 	SendShutdownEvent(Analytics, WatchdogCommandLine, bReturnCodeObtained, ReturnCode, UserResponse, StoredValues, StartupTime);
 
 	// Shutdown tool and engine
diff --git a/Engine/Source/Runtime/Analytics/Analytics/Public/AnalyticsEventAttribute.h b/Engine/Source/Runtime/Analytics/Analytics/Public/AnalyticsEventAttribute.h
--- a/Engine/Source/Runtime/Analytics/Analytics/Public/AnalyticsEventAttribute.h
+++ b/Engine/Source/Runtime/Analytics/Analytics/Public/AnalyticsEventAttribute.h
@@ -5,6 +5,16 @@
 #include "CoreMinimal.h"
 #include "AnalyticsConversion.h"
 
+struct FJsonNull
+{
+};
+
+struct FJsonFragment
+{
+	explicit FJsonFragment(FString&& StringRef) : FragmentString(MoveTemp(StringRef)) {}
+	FString FragmentString;
+};
+
 /**
  * Struct to hold key/value pairs that will be sent as attributes along with analytics events.
  * All values are actually strings, but we provide a convenient constructor that relies on ToStringForAnalytics() to 
@@ -12,12 +22,123 @@
  */
 struct FAnalyticsEventAttribute
 {
-	FString AttrName;
-	FString AttrValue;
+	const FString AttrName;
+
+	const FString AttrValueString;
+	const double AttrValueNumber;
+	const bool AttrValueBool;
+
+	enum class AttrTypeEnum
+	{
+		String,
+		Number,
+		Boolean,
+		Null,
+		JsonFragment
+	};
+	const AttrTypeEnum AttrType;
+
 	/** Default ctor since we declare a custom ctor. */
 	FAnalyticsEventAttribute()
+		: AttrName()
+		, AttrValueString()
+		, AttrValueNumber(0)
+		, AttrValueBool(false)
+		, AttrType(AttrTypeEnum::String)
 	{}
 
+	/** If you need the old AttrValue behavior (i.e. stringify everything), call this function instead. */
+	FString ToString() const
+	{
+		switch (AttrType)
+		{
+		case AttrTypeEnum::String:
+		case AttrTypeEnum::JsonFragment:
+			return AttrValueString;
+		case AttrTypeEnum::Number:
+			return Lex::ToSanitizedString(AttrValueNumber);
+		case AttrTypeEnum::Boolean:
+			return Lex::ToString(AttrValueBool);
+		case AttrTypeEnum::Null:
+			return TEXT("null");
+		default:
+			ensure(false);
+			return FString();
+		}
+	}
+
+public: // null
+	template <typename NameType>
+	FAnalyticsEventAttribute(NameType&& InName, FJsonNull)
+		: AttrName(Lex::ToString(Forward<NameType>(InName)))
+		, AttrValueString()
+		, AttrValueNumber(0)
+		, AttrValueBool(false)
+		, AttrType(AttrTypeEnum::Null)
+	{
+	}
+
+public: // numeric types
+	template <typename NameType>
+	FAnalyticsEventAttribute(NameType&& InName, double InValue)
+		: AttrName(Lex::ToString(Forward<NameType>(InName)))
+		, AttrValueString()
+		, AttrValueNumber(InValue)
+		, AttrValueBool(false)
+		, AttrType(AttrTypeEnum::Number)
+	{
+	}
+	template <typename NameType>
+	FAnalyticsEventAttribute(NameType&& InName, float InValue)
+		: AttrName(Lex::ToString(Forward<NameType>(InName)))
+		, AttrValueString()
+		, AttrValueNumber(InValue)
+		, AttrValueBool(false)
+		, AttrType(AttrTypeEnum::Number)
+	{
+	}
+	template <typename NameType>
+	FAnalyticsEventAttribute(NameType&& InName, int32 InValue)
+		: AttrName(Lex::ToString(Forward<NameType>(InName)))
+		, AttrValueString()
+		, AttrValueNumber(InValue)
+		, AttrValueBool(false)
+		, AttrType(AttrTypeEnum::Number)
+	{
+	}
+	template <typename NameType>
+	FAnalyticsEventAttribute(NameType&& InName, uint32 InValue)
+		: AttrName(Lex::ToString(Forward<NameType>(InName)))
+		, AttrValueString()
+		, AttrValueNumber(InValue)
+		, AttrValueBool(false)
+		, AttrType(AttrTypeEnum::Number)
+	{
+	}
+
+public: // boolean
+	template <typename NameType>
+	FAnalyticsEventAttribute(NameType&& InName, bool InValue)
+		: AttrName(Lex::ToString(Forward<NameType>(InName)))
+		, AttrValueString()
+		, AttrValueNumber(0)
+		, AttrValueBool(InValue)
+		, AttrType(AttrTypeEnum::Boolean)
+	{
+	}
+
+public: // json fragment
+	template <typename NameType>
+	FAnalyticsEventAttribute(NameType&& InName, FJsonFragment&& Fragment)
+		: AttrName(Lex::ToString(Forward<NameType>(InName)))
+		, AttrValueString(MoveTemp(Fragment.FragmentString))
+		, AttrValueNumber(0)
+		, AttrValueBool(false)
+		, AttrType(AttrTypeEnum::JsonFragment)
+	{
+	}
+
+public: // string (catch-all)
 	/**
 	 * Helper constructor to make an attribute from a name/value pair by forwarding through Lex::ToString and AnalyticsConversion::ToString.
 	 * 
@@ -27,7 +148,10 @@ struct FAnalyticsEventAttribute
 	template <typename NameType, typename ValueType>
 	FAnalyticsEventAttribute(NameType&& InName, ValueType&& InValue)
 		: AttrName(Lex::ToString(Forward<NameType>(InName)))
-		, AttrValue(AnalyticsConversion::ToString(Forward<ValueType>(InValue)))
+		, AttrValueString(AnalyticsConversion::ToString(Forward<ValueType>(InValue)))
+		, AttrValueNumber(0)
+		, AttrValueBool(false)
+		, AttrType(AttrTypeEnum::String)
 	{}
 };
 
diff --git a/Engine/Source/Runtime/Analytics/AnalyticsET/Private/IAnalyticsProviderET.cpp b/Engine/Source/Runtime/Analytics/AnalyticsET/Private/IAnalyticsProviderET.cpp
--- a/Engine/Source/Runtime/Analytics/AnalyticsET/Private/IAnalyticsProviderET.cpp
+++ b/Engine/Source/Runtime/Analytics/AnalyticsET/Private/IAnalyticsProviderET.cpp
@@ -194,6 +194,7 @@ public:
 	virtual void RecordEventJson(FString EventName, TArray<FAnalyticsEventAttribute>&& AttributesJson) override;
 	virtual void SetDefaultEventAttributes(TArray<FAnalyticsEventAttribute>&& Attributes) override;
 	virtual const TArray<FAnalyticsEventAttribute>& GetDefaultEventAttributes() const override;
+	virtual void SetEventCallback(const OnEventRecorded& Callback) override;
 
 	virtual ~FAnalyticsProviderET();
 
@@ -281,6 +282,8 @@ private:
 	/** Critical section for updating the CachedEvents. Mutable to allow const methods to access the list. */
 	mutable FCriticalSection CachedEventsCS;
 
+	TArray<OnEventRecorded> EventRecordedCallbacks;
+
 	/**
 	* Delegate called when an event Http request completes
 	*/
@@ -313,6 +316,7 @@ public:
 	virtual void RecordEvent(FString EventName, TArray<FAnalyticsEventAttribute>&& Attributes) override {}
 	virtual void RecordEventJson(FString EventName, TArray<FAnalyticsEventAttribute>&& AttributesJson) override {}
 	virtual void SetDefaultEventAttributes(TArray<FAnalyticsEventAttribute>&& Attributes) override {}
+	virtual void SetEventCallback(const OnEventRecorded& Callback) override {}
 
 	virtual ~FAnalyticsProviderETNULL() {};
 
@@ -558,21 +562,49 @@ void FAnalyticsProviderET::FlushEvents()
 					// default attributes for this event
 					for (const FAnalyticsEventAttribute& Attr : CurrentDefaultAttributes)
 					{
-						JsonWriter->WriteValue(Attr.AttrName, Attr.AttrValue);
+						switch (Attr.AttrType)
+						{
+						case FAnalyticsEventAttribute::AttrTypeEnum::String:
+							JsonWriter->WriteValue(Attr.AttrName, Attr.AttrValueString);
+							break;
+						case FAnalyticsEventAttribute::AttrTypeEnum::Number:
+							JsonWriter->WriteValue(Attr.AttrName, Attr.AttrValueNumber);
+							break;
+						case FAnalyticsEventAttribute::AttrTypeEnum::Boolean:
+							JsonWriter->WriteValue(Attr.AttrName, Attr.AttrValueBool);
+							break;
+						case FAnalyticsEventAttribute::AttrTypeEnum::JsonFragment:
+							JsonWriter->WriteRawJSONValue(Attr.AttrName, Attr.AttrValueString);
+							break;
+						}
 					}
 					// optional attributes for this event
 					if (!Entry.bIsJsonEvent)
 					{
 						for (const FAnalyticsEventAttribute& Attr : Entry.Attributes)
 						{
-							JsonWriter->WriteValue(Attr.AttrName, Attr.AttrValue);
+							switch (Attr.AttrType)
+							{
+							case FAnalyticsEventAttribute::AttrTypeEnum::String:
+								JsonWriter->WriteValue(Attr.AttrName, Attr.AttrValueString);
+								break;
+							case FAnalyticsEventAttribute::AttrTypeEnum::Number:
+								JsonWriter->WriteValue(Attr.AttrName, Attr.AttrValueNumber);
+								break;
+							case FAnalyticsEventAttribute::AttrTypeEnum::Boolean:
+								JsonWriter->WriteValue(Attr.AttrName, Attr.AttrValueBool);
+								break;
+							case FAnalyticsEventAttribute::AttrTypeEnum::JsonFragment:
+								JsonWriter->WriteRawJSONValue(Attr.AttrName, Attr.AttrValueString);
+								break;
+							}
 						}
 					}
 					else
 					{
 						for (const FAnalyticsEventAttribute& Attr : Entry.Attributes)
 						{
-							JsonWriter->WriteRawJSONValue(Attr.AttrName, Attr.AttrValue);
+							JsonWriter->WriteRawJSONValue(Attr.AttrName, Attr.AttrValueString);
 						}
 					}
 					JsonWriter->WriteObjectEnd();
@@ -650,7 +682,7 @@ void FAnalyticsProviderET::FlushEvents()
 							PayloadNdx,
 							*FPlatformHttp::UrlEncode(CurrentDefaultAttributes[DefaultAttributeNdx].AttrName),
 							PayloadNdx,
-							*FPlatformHttp::UrlEncode(CurrentDefaultAttributes[DefaultAttributeNdx].AttrValue));
+							*FPlatformHttp::UrlEncode(CurrentDefaultAttributes[DefaultAttributeNdx].ToString()));
 					}
 					// optional attributes for this event
 					for (int AttrNdx = 0; AttrNdx < Event.Attributes.Num() && PayloadNdx < 40; ++AttrNdx, ++PayloadNdx)
@@ -659,7 +691,7 @@ void FAnalyticsProviderET::FlushEvents()
 							PayloadNdx,
 							*FPlatformHttp::UrlEncode(Event.Attributes[AttrNdx].AttrName),
 							PayloadNdx,
-							*FPlatformHttp::UrlEncode(Event.Attributes[AttrNdx].AttrValue));
+							*FPlatformHttp::UrlEncode(Event.Attributes[AttrNdx].ToString()));
 					}
 
 					// log out the un-encoded values to make reading the log easier.
@@ -766,6 +798,12 @@ void FAnalyticsProviderET::RecordEvent(const FString& EventName, const TArray<FA
 
 void FAnalyticsProviderET::RecordEvent(FString EventName, TArray<FAnalyticsEventAttribute>&& Attributes)
 {
+	// fire any callbacks
+	for (const auto& Cb : EventRecordedCallbacks)
+	{
+		Cb(EventName, Attributes, false);
+	}
+
 	// There are much better ways to do this, but since most events are recorded and handled on the same (game) thread,
 	// this is probably mostly fine for now, and simply favoring not crashing at the moment
 	FScopeLock ScopedLock(&CachedEventsCS);
@@ -780,6 +818,13 @@ void FAnalyticsProviderET::RecordEvent(FString EventName, TArray<FAnalyticsEvent
 void FAnalyticsProviderET::RecordEventJson(FString EventName, TArray<FAnalyticsEventAttribute>&& AttributesJson)
 {
 	checkf(!UseLegacyProtocol, TEXT("Cannot use Json events with legacy protocol"));
+
+	// fire any callbacks
+	for (const auto& Cb : EventRecordedCallbacks)
+	{
+		Cb(EventName, AttributesJson, true);
+	}
+
 	// There are much better ways to do this, but since most events are recorded and handled on the same (game) thread,
 	// this is probably mostly fine for now, and simply favoring not crashing at the moment
 	FScopeLock ScopedLock(&CachedEventsCS);
@@ -815,6 +860,11 @@ const TArray<FAnalyticsEventAttribute>& FAnalyticsProviderET::GetDefaultEventAtt
 	return CachedEvents[DefaultIndex].Attributes;
 }
 
+void FAnalyticsProviderET::SetEventCallback(const OnEventRecorded& Callback)
+{
+	EventRecordedCallbacks.Add(Callback);
+}
+
 void FAnalyticsProviderET::EventRequestComplete(FHttpRequestPtr HttpRequest, FHttpResponsePtr HttpResponse, bool, TSharedPtr< TArray<FAnalyticsEventEntry> > FlushedEvents)
 {
 	// process responses
diff --git a/Engine/Source/Runtime/Analytics/AnalyticsET/Public/IAnalyticsProviderET.h b/Engine/Source/Runtime/Analytics/AnalyticsET/Public/IAnalyticsProviderET.h
--- a/Engine/Source/Runtime/Analytics/AnalyticsET/Public/IAnalyticsProviderET.h
+++ b/Engine/Source/Runtime/Analytics/AnalyticsET/Public/IAnalyticsProviderET.h
@@ -87,4 +87,12 @@ public:
 	* @param Attributes array of attributes that should be appended to every event.
 	*/
 	virtual const TArray<FAnalyticsEventAttribute>& GetDefaultEventAttributes() const = 0;
+
+	/**
+	* Set a callback to be invoked any time an event is queued.
+	* 
+	* @param the callback
+	*/
+	typedef TFunction<void(const FString& EventName, const TArray<FAnalyticsEventAttribute>& Attrs, bool bJson)> OnEventRecorded;
+	virtual void SetEventCallback(const OnEventRecorded& Callback) = 0;
 };
diff --git a/Engine/Source/Runtime/Analytics/AnalyticsSwrve/Private/AnalyticsSwrve.cpp b/Engine/Source/Runtime/Analytics/AnalyticsSwrve/Private/AnalyticsSwrve.cpp
--- a/Engine/Source/Runtime/Analytics/AnalyticsSwrve/Private/AnalyticsSwrve.cpp
+++ b/Engine/Source/Runtime/Analytics/AnalyticsSwrve/Private/AnalyticsSwrve.cpp
@@ -357,7 +357,7 @@ void FAnalyticsProviderSwrve::RecordEvent(const FString& EventName, const TArray
 			{
 				EventParams += TEXT(",");
 			}
-			EventParams += FString(TEXT("\"")) + Attributes[Ndx].AttrName + TEXT("\": \"") + Attributes[Ndx].AttrValue+ TEXT("\"");
+			EventParams += FString(TEXT("\"")) + Attributes[Ndx].AttrName + TEXT("\": \"") + Attributes[Ndx].ToString()+ TEXT("\"");
 		}
 		EventParams += TEXT("}");
 	}
diff --git a/Engine/Source/Runtime/Analytics/QoSReporter/Private/QoSReporterModule.cpp b/Engine/Source/Runtime/Analytics/QoSReporter/Private/QoSReporterModule.cpp
--- a/Engine/Source/Runtime/Analytics/QoSReporter/Private/QoSReporterModule.cpp
+++ b/Engine/Source/Runtime/Analytics/QoSReporter/Private/QoSReporterModule.cpp
@@ -238,7 +238,7 @@ void FAnalyticsProviderQoSReporter::RecordEvent(const FString& InEventName, cons
 			for (int32 AttrIdx = 0; AttrIdx < Attributes.Num(); AttrIdx++)
 			{
 				const FAnalyticsEventAttribute& Attr = Attributes[AttrIdx];
-				JsonWriter->WriteValue(Attr.AttrName, Attr.AttrValue);
+				JsonWriter->WriteValue(Attr.AttrName, Attr.ToString());
 			}
 		}
 		JsonWriter->WriteObjectEnd();
diff --git a/Engine/Source/Runtime/Android/AndroidAudio/Private/AndroidAudioSource.cpp b/Engine/Source/Runtime/Android/AndroidAudio/Private/AndroidAudioSource.cpp
--- a/Engine/Source/Runtime/Android/AndroidAudio/Private/AndroidAudioSource.cpp
+++ b/Engine/Source/Runtime/Android/AndroidAudio/Private/AndroidAudioSource.cpp
@@ -508,6 +508,12 @@ void FSLESSoundSource::Stop( void )
 			check(SL_RESULT_SUCCESS == result);
 		}
 	}
+	DestroyPlayer();
+	ReleaseResources();
+	Paused = false;
+	Playing = false;
+	SLESBuffer = nullptr;
+	Buffer = nullptr;
 
 	DestroyPlayer();
 	ReleaseResources();
diff --git a/Engine/Source/Runtime/Core/Private/Android/AndroidMisc.cpp b/Engine/Source/Runtime/Core/Private/Android/AndroidMisc.cpp
--- a/Engine/Source/Runtime/Core/Private/Android/AndroidMisc.cpp
+++ b/Engine/Source/Runtime/Core/Private/Android/AndroidMisc.cpp
@@ -27,6 +27,7 @@
 
 #include <android_native_app_glue.h>
 #include "Function.h"
+#include "AndroidStats.h"
 
 DECLARE_LOG_CATEGORY_EXTERN(LogEngine, Log, All);
 
@@ -552,7 +553,7 @@ FAndroidMisc::FCPUState& FAndroidMisc::GetCPUState(){
 	int32		Index = 0;
 	ANSICHAR	Buffer[500];
 
-	CurrentCPUState.CoreCount = FAndroidMisc::NumberOfCores();
+	CurrentCPUState.CoreCount = FMath::Min( FAndroidMisc::NumberOfCores(), FAndroidMisc::FCPUState::MaxSupportedCores);
 	FILE* FileHandle = fopen("/proc/stat", "r");
 	if (FileHandle){
 		CurrentCPUState.ActivatedCoreCount = 0;
@@ -563,18 +564,23 @@ FAndroidMisc::FCPUState& FAndroidMisc::GetCPUState(){
 		
 		while (fgets(Buffer, 100, FileHandle)) {
 #if PLATFORM_64BITS
-			sscanf(Buffer, "%4s %8lu %8lu %8lu %8lu %8lu %8lu %8lu", CurrentCPUState.Name,
+			sscanf(Buffer, "%5s %8lu %8lu %8lu %8lu %8lu %8lu %8lu", CurrentCPUState.Name,
 				&UserTime, &NiceTime, &SystemTime, &IdleTime, &IOWaitTime, &IRQTime,
 				&SoftIRQTime);
 #else
-			sscanf(Buffer, "%4s %8llu %8llu %8llu %8llu %8llu %8llu %8llu", CurrentCPUState.Name,
+			sscanf(Buffer, "%5s %8llu %8llu %8llu %8llu %8llu %8llu %8llu", CurrentCPUState.Name,
 				&UserTime, &NiceTime, &SystemTime, &IdleTime, &IOWaitTime, &IRQTime,
 				&SoftIRQTime);
 #endif
 
 			if (0 == strncmp(CurrentCPUState.Name, "cpu", 3)) {
 				Index = CurrentCPUState.Name[3] - '0';
-				if (Index >= 0 && Index < CurrentCPUState.CoreCount) {
+				if (Index >= 0 && Index < CurrentCPUState.CoreCount)
+				{
+					if(CurrentCPUState.Name[5] != '\0')
+					{
+						Index = atol(&CurrentCPUState.Name[3]);
+					}
 					CurrentCPUState.CurrentUsage[Index].IdleTime = IdleTime;
 					CurrentCPUState.CurrentUsage[Index].NiceTime = NiceTime;
 					CurrentCPUState.CurrentUsage[Index].SystemTime = SystemTime;
@@ -1598,3 +1604,38 @@ FString FAndroidMisc::GetOSVersion()
 {
 	return GetAndroidVersion();
 }
+
+uint32 FAndroidMisc::GetCoreFrequency(int32 CoreIndex, ECoreFrequencyProperty CoreFrequencyProperty)
+{
+	uint32 ReturnFrequency = 0;
+	char QueryFile[256];
+	const char* FreqProperty = nullptr;
+	static const char* CurrentFrequencyString = "scaling_cur_freq";
+	static const char* MaxFrequencyString = "cpuinfo_max_freq";
+	static const char* MinFrequencyString = "cpuinfo_min_freq";
+	switch (CoreFrequencyProperty)
+	{
+		case ECoreFrequencyProperty::MaxFrequency:
+			FreqProperty = MaxFrequencyString;
+			break;
+		case ECoreFrequencyProperty::MinFrequency:
+			FreqProperty = MinFrequencyString;
+			break;
+		default:
+		case ECoreFrequencyProperty::CurrentFrequency:
+			FreqProperty = CurrentFrequencyString;
+			break;
+	}
+	sprintf(QueryFile, "/sys/devices/system/cpu/cpu%d/cpufreq/%s", CoreIndex, FreqProperty);
+
+	if (FILE* CoreFreqStateFile = fopen(QueryFile, "r"))
+	{
+		char curr_corefreq[32] = { 0 };
+		if( fgets(curr_corefreq, ARRAY_COUNT(curr_corefreq), CoreFreqStateFile) != nullptr)
+		{
+			ReturnFrequency = atol(curr_corefreq);
+		}
+		fclose(CoreFreqStateFile);
+	}
+	return ReturnFrequency;
+}
diff --git a/Engine/Source/Runtime/Core/Private/Android/AndroidStats.cpp b/Engine/Source/Runtime/Core/Private/Android/AndroidStats.cpp
new file mode 100644
--- /dev/null
+++ b/Engine/Source/Runtime/Core/Private/Android/AndroidStats.cpp
@@ -0,0 +1,201 @@
+// Copyright 1998-2017 Epic Games, Inc. All Rights Reserved.
+
+#include "AndroidStats.h"
+#include "CoreMinimal.h"
+#include "AndroidMisc.h"
+#include "Stats/Stats.h"
+#include "IConsoleManager.h"
+
+DECLARE_STATS_GROUP(TEXT("Android CPU stats"), STATGROUP_AndroidCPU, STATCAT_Advanced);
+
+DECLARE_DWORD_COUNTER_STAT(TEXT("Num Frequency Groups"), STAT_NumFreqGroups, STATGROUP_AndroidCPU);
+DECLARE_DWORD_COUNTER_STAT(TEXT("Freq Group 0 : Max frequency"), STAT_FreqGroup0MaxFrequency, STATGROUP_AndroidCPU);
+DECLARE_DWORD_COUNTER_STAT(TEXT("Freq Group 0 : Min frequency"), STAT_FreqGroup0MinFrequency, STATGROUP_AndroidCPU);
+DECLARE_FLOAT_COUNTER_STAT(TEXT("Freq Group 0 : % of max frequency"), STAT_FreqGroup0CurrentFrequency, STATGROUP_AndroidCPU);
+DECLARE_DWORD_COUNTER_STAT(TEXT("Freq Group 0 : Num Cores"), STAT_FreqGroup0NumCores, STATGROUP_AndroidCPU);
+
+DECLARE_DWORD_COUNTER_STAT(TEXT("Freq Group 1 : Max frequency"), STAT_FreqGroup1MaxFrequency, STATGROUP_AndroidCPU);
+DECLARE_DWORD_COUNTER_STAT(TEXT("Freq Group 1 : Min frequency"), STAT_FreqGroup1MinFrequency, STATGROUP_AndroidCPU);
+DECLARE_FLOAT_COUNTER_STAT(TEXT("Freq Group 1 : % of max frequency"), STAT_FreqGroup1CurrentFrequency, STATGROUP_AndroidCPU);
+DECLARE_DWORD_COUNTER_STAT(TEXT("Freq Group 1 : Num Cores"), STAT_FreqGroup1NumCores, STATGROUP_AndroidCPU);
+
+DECLARE_DWORD_COUNTER_STAT(TEXT("Freq Group 2 : Max frequency"), STAT_FreqGroup2MaxFrequency, STATGROUP_AndroidCPU);
+DECLARE_DWORD_COUNTER_STAT(TEXT("Freq Group 2 : Min frequency"), STAT_FreqGroup2MinFrequency, STATGROUP_AndroidCPU);
+DECLARE_FLOAT_COUNTER_STAT(TEXT("Freq Group 2 : % of max frequency"), STAT_FreqGroup2CurrentFrequency, STATGROUP_AndroidCPU);
+DECLARE_DWORD_COUNTER_STAT(TEXT("Freq Group 2 : Num Cores"), STAT_FreqGroup2NumCores, STATGROUP_AndroidCPU);
+
+DECLARE_DWORD_COUNTER_STAT(TEXT("Freq Group 3 : Max frequency"), STAT_FreqGroup3MaxFrequency, STATGROUP_AndroidCPU);
+DECLARE_DWORD_COUNTER_STAT(TEXT("Freq Group 3 : Min frequency"), STAT_FreqGroup3MinFrequency, STATGROUP_AndroidCPU);
+DECLARE_FLOAT_COUNTER_STAT(TEXT("Freq Group 3 : % of max frequency"), STAT_FreqGroup3CurrentFrequency, STATGROUP_AndroidCPU);
+DECLARE_DWORD_COUNTER_STAT(TEXT("Freq Group 3 : Num Cores"), STAT_FreqGroup3NumCores, STATGROUP_AndroidCPU);
+
+DECLARE_DWORD_COUNTER_STAT(TEXT("Num CPU Cores"), STAT_NumCPUCores, STATGROUP_AndroidCPU);
+
+DECLARE_FLOAT_COUNTER_STAT(TEXT("Freq Group 0 : highest core utilization %"), STAT_FreqGroup0MaxUtilization, STATGROUP_AndroidCPU);
+DECLARE_FLOAT_COUNTER_STAT(TEXT("Freq Group 1 : highest core utilization %"), STAT_FreqGroup1MaxUtilization, STATGROUP_AndroidCPU);
+DECLARE_FLOAT_COUNTER_STAT(TEXT("Freq Group 2 : highest core utilization %"), STAT_FreqGroup2MaxUtilization, STATGROUP_AndroidCPU);
+DECLARE_FLOAT_COUNTER_STAT(TEXT("Freq Group 3 : highest core utilization %"), STAT_FreqGroup3MaxUtilization, STATGROUP_AndroidCPU);
+
+int32 GAndroidCollectCPUStats = 1;
+static FAutoConsoleVariableRef CVarAndroidCollectCPUStats(
+	TEXT("Android.CollectCPUStats"),
+	GAndroidCollectCPUStats,
+	TEXT("Collect CPU Stats (Default: 1)"),
+	ECVF_Default);
+
+#if !STATS
+
+void FAndroidStats::UpdateAndroidStats()
+{
+}
+
+#else
+
+#define SET_DWORD_STAT_BY_FNAME(Stat, Amount) \
+{\
+	if (Amount != 0) \
+		FThreadStats::AddMessage(Stat, EStatOperation::Set, int64(Amount));\
+}
+
+#define SET_FLOAT_STAT_BY_FNAME(Stat, Amount) \
+{\
+	if (Amount != 0) \
+		FThreadStats::AddMessage(Stat, EStatOperation::Set, double(Amount));\
+}
+
+void FAndroidStats::UpdateAndroidStats()
+{
+	if (GAndroidCollectCPUStats == 0)
+	{
+		return;
+	}
+
+	static const FName AndroidFrequencyGroupMaxFreqStats[] = {
+		GET_STATFNAME(STAT_FreqGroup0MaxFrequency),
+		GET_STATFNAME(STAT_FreqGroup1MaxFrequency),
+		GET_STATFNAME(STAT_FreqGroup2MaxFrequency),
+		GET_STATFNAME(STAT_FreqGroup3MaxFrequency),
+	};
+
+	static const FName AndroidFrequencyGroupMinFreqStats[] = {
+		GET_STATFNAME(STAT_FreqGroup0MinFrequency),
+		GET_STATFNAME(STAT_FreqGroup1MinFrequency),
+		GET_STATFNAME(STAT_FreqGroup2MinFrequency),
+		GET_STATFNAME(STAT_FreqGroup3MinFrequency),
+	};
+
+	static const FName AndroidFrequencyGroupCurrentFreqStats[] = {
+		GET_STATFNAME(STAT_FreqGroup0CurrentFrequency),
+		GET_STATFNAME(STAT_FreqGroup1CurrentFrequency),
+		GET_STATFNAME(STAT_FreqGroup2CurrentFrequency),
+		GET_STATFNAME(STAT_FreqGroup3CurrentFrequency),
+	};
+
+	static const FName AndroidFrequencyGroupNumCoresStats[] = {
+		GET_STATFNAME(STAT_FreqGroup0NumCores),
+		GET_STATFNAME(STAT_FreqGroup1NumCores),
+		GET_STATFNAME(STAT_FreqGroup2NumCores),
+		GET_STATFNAME(STAT_FreqGroup3NumCores),
+	};
+
+	static const FName AndroidFrequencyGroupMaxCoresUtilizationStats[] = {
+		GET_STATFNAME(STAT_FreqGroup0MaxUtilization),
+		GET_STATFNAME(STAT_FreqGroup1MaxUtilization),
+		GET_STATFNAME(STAT_FreqGroup2MaxUtilization),
+		GET_STATFNAME(STAT_FreqGroup3MaxUtilization),
+	};
+
+	const int32 MaxCoresStatsSupport = 16;
+	int32 NumCores = FMath::Min(FAndroidMisc::NumberOfCores(), MaxCoresStatsSupport);
+	uint32 CurrentCPUFrequency[MaxCoresStatsSupport];
+
+	struct FFrequencyGroup
+	{
+		uint32 MinFrequency;
+		uint32 MaxFrequency;
+		uint32 CoreCount;
+	};
+
+	static uint32 UnInitializedCores = NumCores;
+	static TArray<FFrequencyGroup> FrequencyGroups;
+	static uint32 CoreFrequencyGroupIndex[MaxCoresStatsSupport] = {
+		0xFFFFFFFF, 0xFFFFFFFF,0xFFFFFFFF,0xFFFFFFFF,
+		0xFFFFFFFF, 0xFFFFFFFF,0xFFFFFFFF,0xFFFFFFFF,
+		0xFFFFFFFF, 0xFFFFFFFF,0xFFFFFFFF,0xFFFFFFFF,
+		0xFFFFFFFF, 0xFFFFFFFF,0xFFFFFFFF,0xFFFFFFFF, };
+
+	if (UnInitializedCores != 0)
+	{
+		for (int32 CoreIndex = 0; CoreIndex < NumCores; CoreIndex++)
+		{
+			if (CoreFrequencyGroupIndex[CoreIndex] == 0xFFFFFFFF)
+			{
+				uint32 MinFreq = FAndroidMisc::GetCoreFrequency(CoreIndex, FAndroidMisc::ECoreFrequencyProperty::MinFrequency);
+				uint32 MaxFreq = FAndroidMisc::GetCoreFrequency(CoreIndex, FAndroidMisc::ECoreFrequencyProperty::MaxFrequency);
+				if (MaxFreq > 0)
+				{
+					UnInitializedCores--;
+					int32 FoundIndex = FrequencyGroups.IndexOfByPredicate([&MinFreq, &MaxFreq](const FFrequencyGroup& s) { return s.MinFrequency == MinFreq && s.MaxFrequency == MaxFreq; });
+					if (FoundIndex == INDEX_NONE)
+					{
+						FFrequencyGroup NewGroup = { MinFreq,MaxFreq, 1 };
+						CoreFrequencyGroupIndex[CoreIndex] = FrequencyGroups.Add(NewGroup);
+					}
+					else
+					{
+						CoreFrequencyGroupIndex[CoreIndex] = FoundIndex;
+						FrequencyGroups[FoundIndex].CoreCount++;
+					}
+				}
+			}
+		}
+	}
+
+	for (int32 FrequencyGroupIndex = 0; FrequencyGroupIndex < FrequencyGroups.Num(); FrequencyGroupIndex++)
+	{
+		const FFrequencyGroup& FrequencyGroup = FrequencyGroups[FrequencyGroupIndex];
+		SET_DWORD_STAT_BY_FNAME(AndroidFrequencyGroupMaxFreqStats[FrequencyGroupIndex], FrequencyGroup.MaxFrequency);
+		SET_DWORD_STAT_BY_FNAME(AndroidFrequencyGroupNumCoresStats[FrequencyGroupIndex], FrequencyGroup.CoreCount);
+		//SET_DWORD_STAT_BY_FNAME(AndroidFrequencyGroupMinFreqStats[FrequencyGroupIndex], FrequencyGroup.MinFrequency);
+	}
+
+	auto GetFrequencyGroupCurrentFrequency = [&](int32 FrequencyGroupIdx)
+	{
+		for (int32 CoreIdx = 0; CoreIdx < NumCores; CoreIdx++)
+		{
+			if (CoreFrequencyGroupIndex[CoreIdx] == FrequencyGroupIdx)
+			{
+				uint32 CoreFreq = FAndroidMisc::GetCoreFrequency(CoreIdx, FAndroidMisc::ECoreFrequencyProperty::CurrentFrequency);
+				if (CoreFreq > 0)
+				{
+					return ((float)CoreFreq / (float)FrequencyGroups[FrequencyGroupIdx].MaxFrequency) * 100.0f;
+				}
+			}
+		}
+		return 0.0f;
+	};
+
+	for (int32 FrequencyGroupIndex = 0; FrequencyGroupIndex < FrequencyGroups.Num(); FrequencyGroupIndex++)
+	{
+		const FFrequencyGroup& FrequencyGroup = FrequencyGroups[FrequencyGroupIndex];
+		SET_FLOAT_STAT_BY_FNAME(AndroidFrequencyGroupCurrentFreqStats[FrequencyGroupIndex], GetFrequencyGroupCurrentFrequency(FrequencyGroupIndex));
+	}
+
+	FAndroidMisc::FCPUState& AndroidCPUState = FAndroidMisc::GetCPUState();
+	static const uint32 MaxFrequencyGroupStats = 4;
+	float MaxSingleCoreUtilization[MaxFrequencyGroupStats] = { 0.0f };
+	for (int32 CoreIndex = 0; CoreIndex < NumCores; CoreIndex++)
+	{
+		uint32 FrequencyGroupIndex = CoreFrequencyGroupIndex[CoreIndex];
+		if (FrequencyGroupIndex != 0xFFFFFFFF)
+		{
+			float& MaxCoreUtilization = MaxSingleCoreUtilization[FrequencyGroupIndex];
+			MaxCoreUtilization = FMath::Max((float)AndroidCPUState.Utilization[CoreIndex], MaxCoreUtilization);
+		}
+	}
+	for (int32 FrequencyGroupIndex = 0; FrequencyGroupIndex < FrequencyGroups.Num(); FrequencyGroupIndex++)
+	{
+		SET_FLOAT_STAT_BY_FNAME(AndroidFrequencyGroupMaxCoresUtilizationStats[FrequencyGroupIndex], MaxSingleCoreUtilization[FrequencyGroupIndex]);
+	}
+}
+#endif
\ No newline at end of file
diff --git a/Engine/Source/Runtime/Core/Private/GenericPlatform/GenericPlatformFile.cpp b/Engine/Source/Runtime/Core/Private/GenericPlatform/GenericPlatformFile.cpp
--- a/Engine/Source/Runtime/Core/Private/GenericPlatform/GenericPlatformFile.cpp
+++ b/Engine/Source/Runtime/Core/Private/GenericPlatform/GenericPlatformFile.cpp
@@ -534,6 +534,52 @@ bool IPlatformFile::IterateDirectoryStatRecursively(const TCHAR* Directory, FDir
 	return IterateDirectoryStat(Directory, Recurse);
 }
 
+class FFindFilesVisitor : public IPlatformFile::FDirectoryVisitor
+{
+public:
+	IPlatformFile&		PlatformFile;
+	TArray<FString>&	FoundFiles;
+	const TCHAR*		FileExtension;
+	int32				FileExtensionLen;
+	FFindFilesVisitor(IPlatformFile& InPlatformFile, TArray<FString>& InFoundFiles, const TCHAR* InFileExtension)
+		: PlatformFile(InPlatformFile)
+		, FoundFiles(InFoundFiles)
+		, FileExtension(InFileExtension)
+		, FileExtensionLen(InFileExtension ? FCString::Strlen(InFileExtension) : 0)
+	{
+	}
+	virtual bool Visit(const TCHAR* FilenameOrDirectory, bool bIsDirectory) override
+	{
+		if (!bIsDirectory)
+		{
+			if (FileExtensionLen > 0)
+			{
+				int32 FileNameLen = FCString::Strlen(FilenameOrDirectory);
+				if (FileNameLen < FileExtensionLen || 
+					FCString::Strcmp(&FilenameOrDirectory[FileNameLen - FileExtensionLen], FileExtension) != 0)
+				{
+					return true;
+				}
+			}
+				
+			FoundFiles.Emplace(FString(FilenameOrDirectory));
+		}
+		return true;
+	}
+};
+
+void IPlatformFile::FindFiles(TArray<FString>& FoundFiles, const TCHAR* Directory, const TCHAR* FileExtension)
+{
+	FFindFilesVisitor FindFilesVisitor(*this, FoundFiles, FileExtension);
+	IterateDirectory(Directory, FindFilesVisitor);
+}
+
+void IPlatformFile::FindFilesRecursively(TArray<FString>& FoundFiles, const TCHAR* Directory, const TCHAR* FileExtension)
+{
+	FFindFilesVisitor FindFilesVisitor(*this, FoundFiles, FileExtension);
+	IterateDirectoryRecursively(Directory, FindFilesVisitor);
+}
+
 bool IPlatformFile::DeleteDirectoryRecursively(const TCHAR* Directory)
 {
 	class FRecurse : public FDirectoryVisitor
diff --git a/Engine/Source/Runtime/Core/Private/GenericPlatform/GenericPlatformMisc.cpp b/Engine/Source/Runtime/Core/Private/GenericPlatform/GenericPlatformMisc.cpp
--- a/Engine/Source/Runtime/Core/Private/GenericPlatform/GenericPlatformMisc.cpp
+++ b/Engine/Source/Runtime/Core/Private/GenericPlatform/GenericPlatformMisc.cpp
@@ -1003,6 +1003,12 @@ void FGenericPlatformMisc::RegisterForRemoteNotifications()
 	// not implemented by default
 }
 
+bool FGenericPlatformMisc::IsRegisteredForRemoteNotifications()
+{
+	// not implemented by default
+	return false;
+}
+
 void FGenericPlatformMisc::UnregisterForRemoteNotifications()
 {
 	// not implemented by default
diff --git a/Engine/Source/Runtime/Core/Private/IOS/IOSAppDelegate.cpp b/Engine/Source/Runtime/Core/Private/IOS/IOSAppDelegate.cpp
--- a/Engine/Source/Runtime/Core/Private/IOS/IOSAppDelegate.cpp
+++ b/Engine/Source/Runtime/Core/Private/IOS/IOSAppDelegate.cpp
@@ -549,7 +549,8 @@ void EngineCrashHandler(const FGenericCrashContext& GenericContext)
 	OSVersion = [[[UIDevice currentDevice] systemVersion] floatValue];
 	if (!FPlatformMisc::IsDebuggerPresent() || GAlwaysReportCrash)
 	{
-        FPlatformMisc::SetCrashHandler(EngineCrashHandler);
+//        FPlatformMisc::SetCrashHandler(EngineCrashHandler);
+        InstallSignalHandlers();
 	}
 
 	// create the main landscape window object
@@ -920,36 +921,48 @@ void EngineCrashHandler(const FGenericCrashContext& GenericContext)
 
 -(void)application : (UIApplication *)application didReceiveRemoteNotification:(NSDictionary *)userInfo fetchCompletionHandler:(void(^)(UIBackgroundFetchResult result))handler
 {
-	NSString* JsonString = @"{}";
-	NSError* JsonError;
-	NSData* JsonData = [NSJSONSerialization dataWithJSONObject : userInfo
-					options : 0
-					error : &JsonError];
-
-	if (JsonData)
+	if (bEngineInit)
 	{
-		JsonString = [[[NSString alloc] initWithData:JsonData encoding : NSUTF8StringEncoding] autorelease];
+		NSString* JsonString = @"{}";
+		NSError* JsonError;
+		NSData* JsonData = [NSJSONSerialization dataWithJSONObject : userInfo
+						options : 0
+						error : &JsonError];
+
+		if (JsonData)
+		{
+			JsonString = [[[NSString alloc] initWithData:JsonData encoding : NSUTF8StringEncoding] autorelease];
+		}
+		
+		FString	jsonFString(JsonString);
+		int AppState;
+		if (application.applicationState == UIApplicationStateInactive)
+		{
+			AppState = 1; // EApplicationState::Inactive;
+		}
+		else if (application.applicationState == UIApplicationStateBackground)
+		{
+			AppState = 2; // EApplicationState::Background;
+		}
+		else
+		{
+			AppState = 3; // EApplicationState::Active;
+		}
+
+		FFunctionGraphTask::CreateAndDispatchWhenReady([jsonFString, AppState]()
+		{
+			FCoreDelegates::ApplicationReceivedRemoteNotificationDelegate.Broadcast(jsonFString, AppState);
+		}, TStatId(), NULL, ENamedThreads::GameThread);
 	}
 	
-	FString	jsonFString(JsonString);
-	int AppState;
-	if (application.applicationState == UIApplicationStateInactive)
-	{
-		AppState = 1; // EApplicationState::Inactive;
-	}
-	else if (application.applicationState == UIApplicationStateBackground)
-	{
-		AppState = 2; // EApplicationState::Background;
-	}
-	else
-	{
-		AppState = 3; // EApplicationState::Active;
-	}
-
-    FFunctionGraphTask::CreateAndDispatchWhenReady([jsonFString, AppState]()
-    {
-		FCoreDelegates::ApplicationReceivedRemoteNotificationDelegate.Broadcast(jsonFString, AppState);
-    }, TStatId(), NULL, ENamedThreads::GameThread);
+	// According to Apple documentation:
+	//   As soon as you finish processing the notification, you must call the
+	//   block in the handler parameter or your app will be terminated. Your app
+	//   has up to 30 seconds of wall-clock time to process the notification and
+	//   call the specified completion handler block. In practice, you should
+	//   call the handler block as soon as you are done processing the
+	//   notification.
+	handler(UIBackgroundFetchResultNoData);
 }
 
 #endif
diff --git a/Engine/Source/Runtime/Core/Private/IOS/IOSPlatformMisc.cpp b/Engine/Source/Runtime/Core/Private/IOS/IOSPlatformMisc.cpp
--- a/Engine/Source/Runtime/Core/Private/IOS/IOSPlatformMisc.cpp
+++ b/Engine/Source/Runtime/Core/Private/IOS/IOSPlatformMisc.cpp
@@ -977,6 +977,16 @@ void FIOSPlatformMisc::RegisterForRemoteNotifications()
 #endif
 }
 
+bool FIOSPlatformMisc::IsRegisteredForRemoteNotifications()
+{
+	return false;
+}
+
+void FIOSPlatformMisc::UnregisterForRemoteNotifications()
+{
+
+}
+
 void FIOSPlatformMisc::GetValidTargetPlatforms(TArray<FString>& TargetPlatformNames)
 {
 	// this is only used to cook with the proper TargetPlatform with COTF, it's not the runtime platform (which is just IOS for both)
diff --git a/Engine/Source/Runtime/Core/Private/IOS/IOSPlatformProcess.cpp b/Engine/Source/Runtime/Core/Private/IOS/IOSPlatformProcess.cpp
--- a/Engine/Source/Runtime/Core/Private/IOS/IOSPlatformProcess.cpp
+++ b/Engine/Source/Runtime/Core/Private/IOS/IOSPlatformProcess.cpp
@@ -41,21 +41,17 @@ void FIOSPlatformProcess::LaunchURL( const TCHAR* URL, const TCHAR* Parms, FStri
 {
 	UE_LOG(LogIOS, Log,  TEXT("LaunchURL %s %s"), URL, Parms?Parms:TEXT("") );
 	NSString* CFUrl = (NSString*)FPlatformString::TCHARToCFString( URL );
-    CFUrl = [CFUrl stringByAddingPercentEncodingWithAllowedCharacters: [NSCharacterSet URLQueryAllowedCharacterSet]];
-	[[UIApplication sharedApplication] openURL: [NSURL URLWithString:( NSString *)CFUrl]];
-
+	bool Result = [[UIApplication sharedApplication] openURL: [NSURL URLWithString: CFUrl]];
 	if (Error != nullptr)
 	{
-		*Error = TEXT("");
+		*Error = Result ? TEXT("") : TEXT("unable to open url");
 	}
 }
 
 bool FIOSPlatformProcess::CanLaunchURL(const TCHAR* URL)
 {
 	NSString* CFUrl = (NSString*)FPlatformString::TCHARToCFString(URL);
-    CFUrl = [CFUrl stringByAddingPercentEncodingWithAllowedCharacters: [NSCharacterSet URLQueryAllowedCharacterSet]];
-	bool Result = [[UIApplication sharedApplication] canOpenURL:[NSURL URLWithString : (NSString *)CFUrl]];
-
+    bool Result = [[UIApplication sharedApplication] canOpenURL:[NSURL URLWithString: CFUrl]];
 	return Result;
 }
 
diff --git a/Engine/Source/Runtime/Core/Public/Android/AndroidMisc.h b/Engine/Source/Runtime/Core/Public/Android/AndroidMisc.h
--- a/Engine/Source/Runtime/Core/Public/Android/AndroidMisc.h
+++ b/Engine/Source/Runtime/Core/Public/Android/AndroidMisc.h
@@ -59,13 +59,14 @@ public:
 
 	struct FCPUState
 	{
+		const static int32			MaxSupportedCores = 16; //Core count 16 is maximum for now
 		int32						CoreCount;
 		int32						ActivatedCoreCount;
-		ANSICHAR					Name[5];
-		FAndroidMisc::FCPUStatTime	CurrentUsage[8]; //Core count 8 is maximum for now
-		FAndroidMisc::FCPUStatTime	PreviousUsage[8];
-		int32						Status[8];
-		double						Utilization[8];
+		ANSICHAR					Name[6];
+		FAndroidMisc::FCPUStatTime	CurrentUsage[MaxSupportedCores]; 
+		FAndroidMisc::FCPUStatTime	PreviousUsage[MaxSupportedCores];
+		int32						Status[MaxSupportedCores];
+		double						Utilization[MaxSupportedCores];
 		double						AverageUtilization;
 		
 	};
@@ -215,6 +216,15 @@ public:
 	static int32 AndroidBuildVersion;
 
 	static bool VolumeButtonsHandledBySystem;
+
+	enum class ECoreFrequencyProperty
+	{
+		CurrentFrequency,
+		MaxFrequency,
+		MinFrequency,
+	};
+
+	static uint32 GetCoreFrequency(int32 CoreIndex, ECoreFrequencyProperty CoreFrequencyProperty);
 };
 
 typedef FAndroidMisc FPlatformMisc;
diff --git a/Engine/Source/Runtime/Core/Public/Android/AndroidStats.h b/Engine/Source/Runtime/Core/Public/Android/AndroidStats.h
new file mode 100644
--- /dev/null
+++ b/Engine/Source/Runtime/Core/Public/Android/AndroidStats.h
@@ -0,0 +1,12 @@
+// Copyright 1998-2017 Epic Games, Inc. All Rights Reserved.
+
+#pragma once
+
+#include "CoreMinimal.h"
+#include "AndroidMisc.h"
+
+class FAndroidStats
+{
+public:
+	static void UpdateAndroidStats();
+};
diff --git a/Engine/Source/Runtime/Core/Public/GenericPlatform/GenericPlatformFile.h b/Engine/Source/Runtime/Core/Public/GenericPlatform/GenericPlatformFile.h
--- a/Engine/Source/Runtime/Core/Public/GenericPlatform/GenericPlatformFile.h
+++ b/Engine/Source/Runtime/Core/Public/GenericPlatform/GenericPlatformFile.h
@@ -371,6 +371,24 @@ public:
 	 * @return				false if the directory did not exist or if the visitor returned false.
 	**/
 	virtual bool IterateDirectoryStatRecursively(const TCHAR* Directory, FDirectoryStatVisitor& Visitor);
+		
+	/**
+	 * Finds all the files within the given directory, with optional file extension filter
+	 * @param Directory			The directory to iterate the contents of
+	 * @param FileExtension		If FileExtension is NULL, or an empty string "" then all files are found.
+	 * 							Otherwise FileExtension can be of the form .EXT or just EXT and only files with that extension will be returned.
+	 * @return FoundFiles		All the files that matched the optional FileExtension filter, or all files if none was specified.
+	 */
+	virtual void FindFiles(TArray<FString>& FoundFiles, const TCHAR* Directory, const TCHAR* FileExtension);
+
+	/**
+	 * Finds all the files within the directory tree, with optional file extension filter
+	 * @param Directory			The starting directory to iterate the contents. This function explores subdirectories
+	 * @param FileExtension		If FileExtension is NULL, or an empty string "" then all files are found.
+	 * 							Otherwise FileExtension can be of the form .EXT or just EXT and only files with that extension will be returned.
+	 * @return FoundFiles		All the files that matched the optional FileExtension filter, or all files if none was specified.
+	 */
+	virtual void FindFilesRecursively(TArray<FString>& FoundFiles, const TCHAR* Directory, const TCHAR* FileExtension);
 
 	/** 
 	 * Delete all files and subdirectories in a directory, then delete the directory itself
diff --git a/Engine/Source/Runtime/Core/Public/GenericPlatform/GenericPlatformMisc.h b/Engine/Source/Runtime/Core/Public/GenericPlatform/GenericPlatformMisc.h
--- a/Engine/Source/Runtime/Core/Public/GenericPlatform/GenericPlatformMisc.h
+++ b/Engine/Source/Runtime/Core/Public/GenericPlatform/GenericPlatformMisc.h
@@ -1023,6 +1023,11 @@ public:
 	 */
 	static void RegisterForRemoteNotifications();
 
+	/**
+	 * Returns whether or not the device has been registered to receive remote notifications.
+	 */
+	static bool IsRegisteredForRemoteNotifications();
+
 	/**
 	* Requests unregistering from receiving remote notifications on the user's device.
 	*/
diff --git a/Engine/Source/Runtime/Core/Public/GenericPlatform/GenericPlatformProcess.h b/Engine/Source/Runtime/Core/Public/GenericPlatform/GenericPlatformProcess.h
--- a/Engine/Source/Runtime/Core/Public/GenericPlatform/GenericPlatformProcess.h
+++ b/Engine/Source/Runtime/Core/Public/GenericPlatform/GenericPlatformProcess.h
@@ -304,7 +304,8 @@ struct CORE_API FGenericPlatformProcess
 	/**
 	 * Launch a uniform resource locator (i.e. http://www.epicgames.com/unreal).
 	 * This is expected to return immediately as the URL is launched by another
-	 * task.
+	 * task. The URL param must already be a valid URL. If you're looking for code 
+	 * to properly escape a URL fragment, use FGenericPlatformHttp::UrlEncode.
 	 */
 	static void LaunchURL( const TCHAR* URL, const TCHAR* Parms, FString* Error );
 
diff --git a/Engine/Source/Runtime/Core/Public/HAL/IPlatformFileCachedWrapper.h b/Engine/Source/Runtime/Core/Public/HAL/IPlatformFileCachedWrapper.h
--- a/Engine/Source/Runtime/Core/Public/HAL/IPlatformFileCachedWrapper.h
+++ b/Engine/Source/Runtime/Core/Public/HAL/IPlatformFileCachedWrapper.h
@@ -368,6 +368,14 @@ public:
 	{
 		return LowerLevel->IterateDirectoryStatRecursively(Directory, Visitor);
 	}
+	virtual void		FindFiles(TArray<FString>& FoundFiles, const TCHAR* Directory, const TCHAR* FileExtension)
+	{
+		return LowerLevel->FindFiles(FoundFiles, Directory, FileExtension);
+	}
+	virtual void		FindFilesRecursively(TArray<FString>& FoundFiles, const TCHAR* Directory, const TCHAR* FileExtension)
+	{
+		return LowerLevel->FindFilesRecursively(FoundFiles, Directory, FileExtension);
+	}
 	virtual bool		DeleteDirectoryRecursively(const TCHAR* Directory) override
 	{
 		return LowerLevel->DeleteDirectoryRecursively(Directory);
diff --git a/Engine/Source/Runtime/Core/Public/IOS/IOSPlatform.h b/Engine/Source/Runtime/Core/Public/IOS/IOSPlatform.h
--- a/Engine/Source/Runtime/Core/Public/IOS/IOSPlatform.h
+++ b/Engine/Source/Runtime/Core/Public/IOS/IOSPlatform.h
@@ -62,7 +62,7 @@ typedef FIOSPlatformTypes FPlatformTypes;
 #define VARARGS															/* Functions with variable arguments */
 #define CDECL															/* Standard C function */
 #define STDCALL															/* Standard calling convention */
-#if UE_BUILD_DEBUG
+#if UE_BUILD_DEBUG || UE_DISABLE_FORCE_INLINE
 #define FORCEINLINE inline 												/* Don't force code to be inline */
 #else
 #define FORCEINLINE inline __attribute__ ((always_inline))				/* Force code to be inline */
diff --git a/Engine/Source/Runtime/Core/Public/IOS/IOSPlatformMisc.h b/Engine/Source/Runtime/Core/Public/IOS/IOSPlatformMisc.h
--- a/Engine/Source/Runtime/Core/Public/IOS/IOSPlatformMisc.h
+++ b/Engine/Source/Runtime/Core/Public/IOS/IOSPlatformMisc.h
@@ -147,6 +147,8 @@ struct CORE_API FIOSPlatformMisc : public FGenericPlatformMisc
 	static bool IsRunningOnBattery();
 
 	static void RegisterForRemoteNotifications();
+	static bool IsRegisteredForRemoteNotifications();
+	static void UnregisterForRemoteNotifications();
 
 	static class IPlatformChunkInstall* GetPlatformChunkInstall();
 	
diff --git a/Engine/Source/Runtime/Core/Public/Logging/LogMacros.h b/Engine/Source/Runtime/Core/Public/Logging/LogMacros.h
--- a/Engine/Source/Runtime/Core/Public/Logging/LogMacros.h
+++ b/Engine/Source/Runtime/Core/Public/Logging/LogMacros.h
@@ -39,7 +39,7 @@ struct CORE_API FMsg
 ----------------------------------------------------------------------------*/
 
 #ifndef COMPILED_IN_MINIMUM_VERBOSITY
-	#define COMPILED_IN_MINIMUM_VERBOSITY VeryVerbose
+	#define COMPILED_IN_MINIMUM_VERBOSITY Log
 #else
 	#if !IS_MONOLITHIC
 		#error COMPILED_IN_MINIMUM_VERBOSITY can only be defined in monolithic builds.
diff --git a/Engine/Source/Runtime/CoreUObject/Private/Misc/ExclusiveLoadPackageTimeTracker.cpp b/Engine/Source/Runtime/CoreUObject/Private/Misc/ExclusiveLoadPackageTimeTracker.cpp
--- a/Engine/Source/Runtime/CoreUObject/Private/Misc/ExclusiveLoadPackageTimeTracker.cpp
+++ b/Engine/Source/Runtime/CoreUObject/Private/Misc/ExclusiveLoadPackageTimeTracker.cpp
@@ -1,6 +1,10 @@
 // Copyright 1998-2017 Epic Games, Inc. All Rights Reserved.
 
 #include "Misc/ExclusiveLoadPackageTimeTracker.h"
+#include "Misc/OutputDeviceArchiveWrapper.h"
+#include "Misc/Paths.h"
+#include "HAL/FileManager.h"
+#include "ProfilingDebugging/ProfilingHelpers.h"
 #include "Misc/ScopeLock.h"
 #include "UObject/UObjectHash.h"
 #include "UObject/Class.h"
@@ -101,10 +105,33 @@ void FExclusiveLoadPackageTimeTracker::InternalPopLoadPackage(UPackage* LoadedPa
 	}
 }
 
-void FExclusiveLoadPackageTimeTracker::InternalDumpReport() const
+void FExclusiveLoadPackageTimeTracker::InternalDumpReport(const TArray<FString>& Args) const
 {
 	FScopeLock Lock(&TimesCritical);
 
+	const bool bLogOutputToFile = Args.ContainsByPredicate([](const FString& Arg) { return Arg.Compare(TEXT("FILE"), ESearchCase::IgnoreCase) == 0; });
+	const bool bAlphaSort       = Args.ContainsByPredicate([](const FString& Arg) { return Arg.Compare(TEXT("-ALPHASORT"), ESearchCase::IgnoreCase) == 0; });
+
+	FOutputDevice* ReportAr = GLog;
+	FArchive* FileAr = nullptr;
+	FOutputDeviceArchiveWrapper* FileArWrapper = nullptr;
+	FString FilenameFull;
+
+	if (bLogOutputToFile)
+	{
+		const FString PathName = *(FPaths::ProfilingDir() + TEXT("LoadReports/"));
+		IFileManager::Get().MakeDirectory(*PathName);
+
+		const FString Filename = CreateProfileFilename(TEXT(""), TEXT(".loadreport"), true);
+		FilenameFull = PathName + Filename;
+
+		FileAr = IFileManager::Get().CreateDebugFileWriter(*FilenameFull);
+		FileArWrapper = new FOutputDeviceArchiveWrapper(FileAr);
+		ReportAr = FileArWrapper;
+
+		UE_LOG(LogLoad, Log, TEXT("LoadTimes.DumpReport: saving to %s"), *FilenameFull);
+	}
+
 	/** Struct to assemble times for assets loaded by class */
 	struct FTimeCount
 	{
@@ -148,31 +175,50 @@ void FExclusiveLoadPackageTimeTracker::InternalDumpReport() const
 
 	const FLoadTime* EndLoadTime = LoadTimes.Find(EndLoadName);
 	const int32 NumNonAssetTimes = EndLoadTime ? 1 : 0;
-	UE_LOG(LogLoad, Log, TEXT("Loaded: %d packages"), LoadTimes.Num() - NumNonAssetTimes);
-	UE_LOG(LogLoad, Log, TEXT("Total time loading packages: %.3f seconds"), TotalLoadTime);
-	UE_LOG(LogLoad, Log, TEXT("Time spent loading assets slower than %.1fms: %.3f seconds"), SlowAssetThreshold * 1000, SlowAssetTime);
-	UE_LOG(LogLoad, Log, TEXT("Slowest asset: %s (%.1fms)"), *LongestLoadName.ToString(), LongestLoadTime * 1000);
-	UE_LOG(LogLoad, Log, TEXT("Time spent in EndLoad: %.3f seconds"), EndLoadTime ? EndLoadTime->ExclusiveTime : 0.f);
-	UE_LOG(LogLoad, Log, TEXT("Time spent in overhead tracking asset load times: %.6f seconds"), TrackerOverhead);
-
-	UE_LOG(LogLoad, Log, TEXT("Dumping asset type load times sorted by exclusive time:"));
+	ReportAr->Logf(TEXT("Loaded: %d packages"), LoadTimes.Num() - NumNonAssetTimes);
+	ReportAr->Logf(TEXT("Total time loading packages: %.3f seconds"), TotalLoadTime);
+	ReportAr->Logf(TEXT("Time spent loading assets slower than %.1fms: %.3f seconds"), SlowAssetThreshold * 1000, SlowAssetTime);
+	ReportAr->Logf(TEXT("Slowest asset: %s (%.1fms)"), *LongestLoadName.ToString(), LongestLoadTime * 1000);
+	ReportAr->Logf(TEXT("Time spent in EndLoad: %.3f seconds"), EndLoadTime ? EndLoadTime->ExclusiveTime : 0.f);
+	ReportAr->Logf(TEXT("Time spent in overhead tracking asset load times: %.6f seconds"), TrackerOverhead);
+
+	ReportAr->Logf(TEXT("Dumping asset type load times sorted by exclusive time:"));
 	TArray<FTimeCount> SortedAssetTypeLoadTimes;
 	AssetTypeLoadTimes.GenerateValueArray(SortedAssetTypeLoadTimes);
 	SortedAssetTypeLoadTimes.Sort([](const FTimeCount& A, const FTimeCount& B){ return A.ExclusiveTime >= B.ExclusiveTime; });
 	for (const FTimeCount& TimeCount : SortedAssetTypeLoadTimes)
 	{
-		UE_LOG(LogLoad, Log, TEXT("    %.3f: %s (%d packages, %.1fms per package)"), TimeCount.ExclusiveTime, *TimeCount.AssetClass.ToString(), TimeCount.Count, TimeCount.ExclusiveTime / TimeCount.Count * 1000);
+		ReportAr->Logf(TEXT("    %.3f: %s (%d packages, %.1fms per package)"), TimeCount.ExclusiveTime, *TimeCount.AssetClass.ToString(), TimeCount.Count, TimeCount.ExclusiveTime / TimeCount.Count * 1000);
 	}
-
+	
 	// Assets faster than this will be excluded
-	const double LowTimeThreshold = 0.05;
+	double LowTimeThreshold = 0.05;
+
+	const FString* LowTimeArg = Args.FindByPredicate([](const FString& Arg) { return Arg.StartsWith(TEXT("LOWTIME="), ESearchCase::IgnoreCase); });
+	if (LowTimeArg != nullptr)
+	{
+		float LowTimeThresholdParam = 0.0F;
+		if (FParse::Value(**LowTimeArg, TEXT("LOWTIME="), LowTimeThresholdParam))
+		{
+			LowTimeThreshold = (double)LowTimeThresholdParam;
+		}
+		
+	}
 
 	TArray<FLoadTime> SortedLoadTimes;
 	LoadTimes.GenerateValueArray(SortedLoadTimes);
 	
 	{
-		UE_LOG(LogLoad, Log, TEXT("Dumping all loaded assets by exclusive load time:"));
-		SortedLoadTimes.Sort([](const FLoadTime& A, const FLoadTime& B){ return A.ExclusiveTime >= B.ExclusiveTime; });
+		ReportAr->Logf(TEXT("Dumping all loaded assets by exclusive load time:"));
+		if (bAlphaSort)
+		{
+			SortedLoadTimes.Sort([](const FLoadTime& A, const FLoadTime& B) { return A.TimeName < B.TimeName; });
+		}
+		else
+		{
+			SortedLoadTimes.Sort([](const FLoadTime& A, const FLoadTime& B) { return A.ExclusiveTime >= B.ExclusiveTime; });
+		}
+
 		int32 LowThresholdCount = 0;
 		double TotalLowTime = 0;
 		for (int32 LoadTimeIdx = 0; LoadTimeIdx < SortedLoadTimes.Num(); ++LoadTimeIdx)
@@ -182,7 +228,7 @@ void FExclusiveLoadPackageTimeTracker::InternalDumpReport() const
 			{
 				if (LoadTime.ExclusiveTime > LowTimeThreshold)
 				{
-					UE_LOG(LogLoad, Log, TEXT("    %.1fms: %s"), LoadTime.ExclusiveTime * 1000, *LoadTime.TimeName.ToString());
+					ReportAr->Logf(TEXT("    %.1fms: %s"), LoadTime.ExclusiveTime * 1000, *LoadTime.TimeName.ToString());
 				}
 				else
 				{
@@ -194,13 +240,21 @@ void FExclusiveLoadPackageTimeTracker::InternalDumpReport() const
 
 		if (LowThresholdCount > 0)
 		{
-			UE_LOG(LogLoad, Log, TEXT("    ... skipped %d assets that loaded in less than %.1fms totaling %.1fms"), LowThresholdCount, LowTimeThreshold * 1000, TotalLowTime * 1000);
+			ReportAr->Logf(TEXT("    ... skipped %d assets that loaded in less than %.1fms totaling %.1fms"), LowThresholdCount, LowTimeThreshold * 1000, TotalLowTime * 1000);
 		}
 	}
 
 	{
-		UE_LOG(LogLoad, Log, TEXT("Dumping all loaded assets by inclusive load time:"));
-		SortedLoadTimes.Sort([](const FLoadTime& A, const FLoadTime& B){ return A.InclusiveTime >= B.InclusiveTime; });
+		ReportAr->Logf(TEXT("Dumping all loaded assets by inclusive load time:"));
+		if (bAlphaSort)
+		{
+			SortedLoadTimes.Sort([](const FLoadTime& A, const FLoadTime& B) { return A.TimeName < B.TimeName; });
+		}
+		else
+		{
+			SortedLoadTimes.Sort([](const FLoadTime& A, const FLoadTime& B) { return A.InclusiveTime >= B.InclusiveTime; });
+		}
+		
 		int32 LowThresholdCount = 0;
 		double TotalLowTime = 0;
 		for (int32 LoadTimeIdx = 0; LoadTimeIdx < SortedLoadTimes.Num(); ++LoadTimeIdx)
@@ -210,7 +264,7 @@ void FExclusiveLoadPackageTimeTracker::InternalDumpReport() const
 			{
 				if (LoadTime.InclusiveTime > LowTimeThreshold)
 				{
-					UE_LOG(LogLoad, Log, TEXT("    %.1fms: %s"), LoadTime.InclusiveTime * 1000, *LoadTime.TimeName.ToString());
+					ReportAr->Logf(TEXT("    %.1fms: %s"), LoadTime.InclusiveTime * 1000, *LoadTime.TimeName.ToString());
 				}
 				else
 				{
@@ -222,9 +276,16 @@ void FExclusiveLoadPackageTimeTracker::InternalDumpReport() const
 
 		if (LowThresholdCount > 0)
 		{
-			UE_LOG(LogLoad, Log, TEXT("    ... skipped %d assets slower than %.1fms totaling %.1fms"), LowThresholdCount, LowTimeThreshold * 1000, TotalLowTime * 1000);
+			ReportAr->Logf(TEXT("    ... skipped %d assets slower than %.1fms totaling %.1fms"), LowThresholdCount, LowTimeThreshold * 1000, TotalLowTime * 1000);
 		}
 	}
+
+	if (FileArWrapper != nullptr)
+	{
+		FileArWrapper->TearDown();
+		delete FileArWrapper;
+		delete FileAr;
+	}
 }
 
 void FExclusiveLoadPackageTimeTracker::InternalResetReport()
@@ -265,7 +326,7 @@ bool FExclusiveLoadPackageTimeTracker::IsPackageLoadTime(const FLoadTime& Time)
 
 void FExclusiveLoadPackageTimeTracker::DumpReportCommandHandler(const TArray<FString>& Args)
 {
-	DumpReport();
+	DumpReport(Args);
 }
 
 void FExclusiveLoadPackageTimeTracker::ResetReportCommandHandler(const TArray<FString>& Args)
diff --git a/Engine/Source/Runtime/CoreUObject/Public/Misc/ExclusiveLoadPackageTimeTracker.h b/Engine/Source/Runtime/CoreUObject/Public/Misc/ExclusiveLoadPackageTimeTracker.h
--- a/Engine/Source/Runtime/CoreUObject/Public/Misc/ExclusiveLoadPackageTimeTracker.h
+++ b/Engine/Source/Runtime/CoreUObject/Public/Misc/ExclusiveLoadPackageTimeTracker.h
@@ -123,10 +123,10 @@ public:
 	}
 
 	/** Displays the data gathered in various ways. */
-	FORCEINLINE static void DumpReport()
+	FORCEINLINE static void DumpReport(const TArray<FString>& Args)
 	{
 #if WITH_LOADPACKAGE_TIME_TRACKER
-		Get().InternalDumpReport();
+		Get().InternalDumpReport(Args);
 #endif
 	}
 
@@ -196,7 +196,7 @@ private:
 	COREUOBJECT_API void InternalPopLoadPackage(UPackage* LoadedPackage, UObject* LoadedAsset);
 
 	/** Displays the data gathered in various ways. */
-	COREUOBJECT_API void InternalDumpReport() const;
+	COREUOBJECT_API void InternalDumpReport(const TArray<FString>& Args) const;
 
 	/** Resets the data */
 	COREUOBJECT_API void InternalResetReport();
diff --git a/Engine/Source/Runtime/Engine/Private/PhysicsEngine/PhysicsAsset.cpp b/Engine/Source/Runtime/Engine/Private/PhysicsEngine/PhysicsAsset.cpp
--- a/Engine/Source/Runtime/Engine/Private/PhysicsEngine/PhysicsAsset.cpp
+++ b/Engine/Source/Runtime/Engine/Private/PhysicsEngine/PhysicsAsset.cpp
@@ -332,7 +332,9 @@ FBox UPhysicsAsset::CalcAABB(const USkinnedMeshComponent* MeshComp, const FTrans
 		Box = FBox( LocalToWorld.GetLocation(), LocalToWorld.GetLocation() );
 	}
 
-	const float MinBoundSize = 1.f;
+	// @todo: - DO NOT copy this change to UE4/Main
+	const float MinBoundSize = 30.f;
+	// @todo: - DO NOT copy this change to UE4/Main
 	const FVector BoxSize = Box.GetSize();
 
 	if(BoxSize.GetMin() < MinBoundSize)
diff --git a/Engine/Source/Runtime/Engine/Private/Tests/AnalyticsTest.cpp b/Engine/Source/Runtime/Engine/Private/Tests/AnalyticsTest.cpp
--- a/Engine/Source/Runtime/Engine/Private/Tests/AnalyticsTest.cpp
+++ b/Engine/Source/Runtime/Engine/Private/Tests/AnalyticsTest.cpp
@@ -102,29 +102,29 @@ bool FAnalyticsEventAttributeUnitTest::RunTest(const FString& Parameters)
 
 		//Setup the 'Event Attributes'
 		FAnalyticsEventAttribute EventAttributesFStringTest(AttributeName, AttributeValue);
-		TestTrue(TEXT("Expected to take in these type of values '(const FString&, const FString&)'"), (EventAttributesFStringTest.AttrName == TEXT("Test of") && EventAttributesFStringTest.AttrValue == "FAnalyticsEventAttribute '(const FString InName, const FString& InVaue)'"));
+		TestTrue(TEXT("Expected to take in these type of values '(const FString&, const FString&)'"), (EventAttributesFStringTest.AttrName == TEXT("Test of") && EventAttributesFStringTest.AttrValueString == "FAnalyticsEventAttribute '(const FString InName, const FString& InVaue)'"));
 
 		FAnalyticsEventAttribute EventAttributesTCHARTest(AttributeName, TEXT("FAnalyticsEventAttribute '(const FString InName, const TCHAR* InValue)'"));
-		TestTrue(TEXT("Expected to take in these type of values '(const FString&, const TCHAR*')"), (EventAttributesTCHARTest.AttrName == TEXT("Test of") && EventAttributesTCHARTest.AttrValue == TEXT("FAnalyticsEventAttribute '(const FString InName, const TCHAR* InValue)'")));
+		TestTrue(TEXT("Expected to take in these type of values '(const FString&, const TCHAR*')"), (EventAttributesTCHARTest.AttrName == TEXT("Test of") && EventAttributesTCHARTest.AttrValueString == TEXT("FAnalyticsEventAttribute '(const FString InName, const TCHAR* InValue)'")));
 
 		bool bIsBoolTest = true;
 		FAnalyticsEventAttribute EventAttributesBoolTest(AttributeName, bIsBoolTest);
-		TestTrue(TEXT("Expected to take in these types of values '(const FString&, bool)'"), (EventAttributesBoolTest.AttrName == TEXT("Test of") && EventAttributesBoolTest.AttrValue == "true"));
+		TestTrue(TEXT("Expected to take in these types of values '(const FString&, bool)'"), (EventAttributesBoolTest.AttrName == TEXT("Test of") && EventAttributesBoolTest.AttrValueBool == bIsBoolTest && EventAttributesBoolTest.AttrValueString.IsEmpty()));
 
 		FGuid GuidTest(FGuid::NewGuid());
 		FAnalyticsEventAttribute EventAttributesGuidTest(AttributeName, GuidTest);
-		TestTrue(TEXT("Expected to take in these type of values '(const FString&, FGuid)'"), (EventAttributesGuidTest.AttrName == TEXT("Test of") && EventAttributesGuidTest.AttrValue == GuidTest.ToString()));
+		TestTrue(TEXT("Expected to take in these type of values '(const FString&, FGuid)'"), (EventAttributesGuidTest.AttrName == TEXT("Test of") && EventAttributesGuidTest.AttrValueString == GuidTest.ToString()));
 
 		int32 InNumericType = 42;
 		FAnalyticsEventAttribute EventAttributesTInValueTest(AttributeName, InNumericType);
-		TestTrue(TEXT("Expected to take in a arithmetic type (example int32)"), (EventAttributesTInValueTest.AttrName == TEXT("Test of") && EventAttributesTInValueTest.AttrValue == TEXT("42")));
+		TestTrue(TEXT("Expected to take in a arithmetic type (example int32)"), (EventAttributesTInValueTest.AttrName == TEXT("Test of") && EventAttributesTInValueTest.AttrValueNumber == InNumericType && EventAttributesTInValueTest.AttrValueString.IsEmpty()));
 
 		TArray<int32> InNumericArray;
 		InNumericArray.Add(0);
 		InNumericArray.Add(1);
 		InNumericArray.Add(2);
 		FAnalyticsEventAttribute EventAttributesTArray(AttributeName, InNumericArray);
-		TestTrue(TEXT("Expected to take in am arithmetic TArray"), (EventAttributesTArray.AttrName == TEXT("Test of") && EventAttributesTArray.AttrValue == TEXT("0,1,2")));
+		TestTrue(TEXT("Expected to take in am arithmetic TArray"), (EventAttributesTArray.AttrName == TEXT("Test of") && EventAttributesTArray.AttrValueString == TEXT("0,1,2")));
 
 		FString TMapKey1 = "TestKey 1";
 		FString TMapKey2 = "TestKey 2";
@@ -137,7 +137,7 @@ bool FAnalyticsEventAttributeUnitTest::RunTest(const FString& Parameters)
 		InTMap.Add(TMapKey2, NumericValue2);
 		InTMap.Add(TMapKey3, NumericValue3);
 		FAnalyticsEventAttribute EventAttributesTMap(AttributeName, InTMap);
-		TestTrue(TEXT("Expected to take in a TMap "), (EventAttributesTMap.AttrName == TEXT("Test of") && EventAttributesTMap.AttrValue == TEXT("TestKey 1:0,TestKey 2:1,TestKey 3:99")));
+		TestTrue(TEXT("Expected to take in a TMap "), (EventAttributesTMap.AttrName == TEXT("Test of") && EventAttributesTMap.AttrValueString == TEXT("TestKey 1:0,TestKey 2:1,TestKey 3:99")));
 
 		return true;
 	}
diff --git a/Engine/Source/Runtime/JsonUtilities/Private/JsonObjectConverter.cpp b/Engine/Source/Runtime/JsonUtilities/Private/JsonObjectConverter.cpp
--- a/Engine/Source/Runtime/JsonUtilities/Private/JsonObjectConverter.cpp
+++ b/Engine/Source/Runtime/JsonUtilities/Private/JsonObjectConverter.cpp
@@ -727,3 +727,45 @@ bool FJsonObjectConverter::JsonAttributesToUStruct(const TMap< FString, TSharedP
 	return true;
 }
 
+FFormatNamedArguments FJsonObjectConverter::ParseTextArgumentsFromJson(const TSharedPtr<const FJsonObject>& JsonObject)
+{
+	FFormatNamedArguments NamedArgs;
+	if (JsonObject.IsValid())
+	{
+		for (const auto& It : JsonObject->Values)
+		{
+			if (!It.Value.IsValid())
+				continue;
+
+			switch (It.Value->Type)
+			{
+			case EJson::Number:
+				// number
+				NamedArgs.Emplace(It.Key, It.Value->AsNumber());
+				break;
+			case EJson::String:
+				// culture invariant string
+				NamedArgs.Emplace(It.Key, FText::FromString(It.Value->AsString()));
+				break;
+			case EJson::Object:
+			{
+				// localized string
+				FText TextOut;
+				if (FJsonObjectConverter::GetTextFromObject(It.Value->AsObject().ToSharedRef(), TextOut))
+				{
+					NamedArgs.Emplace(It.Key, TextOut);
+				}
+				else
+				{
+					UE_LOG(LogJson, Error, TEXT("Unable to apply Json parameter %s (could not parse object)"), *It.Key);
+				}
+			}
+			break;
+			default:
+				UE_LOG(LogJson, Error, TEXT("Unable to apply Json parameter %s (bad type)"), *It.Key);
+				break;
+			}
+		}
+	}
+	return NamedArgs;
+}
\ No newline at end of file
diff --git a/Engine/Source/Runtime/JsonUtilities/Private/JsonObjectWrapper.cpp b/Engine/Source/Runtime/JsonUtilities/Private/JsonObjectWrapper.cpp
--- a/Engine/Source/Runtime/JsonUtilities/Private/JsonObjectWrapper.cpp
+++ b/Engine/Source/Runtime/JsonUtilities/Private/JsonObjectWrapper.cpp
@@ -68,7 +68,7 @@ bool FJsonObjectWrapper::ExportTextItem(FString& ValueStr, FJsonObjectWrapper co
 
 void FJsonObjectWrapper::PostSerialize(const FArchive& Ar)
 {
-	if (Ar.IsLoading())
+	if (!JsonString.IsEmpty())
 	{
 		// try to parse JsonString
 		TSharedRef<TJsonReader<> > JsonReader = TJsonReaderFactory<>::Create(JsonString);
diff --git a/Engine/Source/Runtime/JsonUtilities/Public/JsonObjectConverter.h b/Engine/Source/Runtime/JsonUtilities/Public/JsonObjectConverter.h
--- a/Engine/Source/Runtime/JsonUtilities/Public/JsonObjectConverter.h
+++ b/Engine/Source/Runtime/JsonUtilities/Public/JsonObjectConverter.h
@@ -8,12 +8,13 @@
 #include "Dom/JsonObject.h"
 #include "Serialization/JsonReader.h"
 #include "Serialization/JsonSerializer.h"
+#include "JsonObjectWrapper.h"
 
 /** Class that handles converting Json objects to and from UStructs */
 class JSONUTILITIES_API FJsonObjectConverter
 {
 public:
-	
+
 	/** FName case insensitivity can make the casing of UPROPERTIES unpredictable. Attempt to standardize output. */
 	static FString StandardizeCase(const FString &StringIn);
 
@@ -31,7 +32,7 @@ public: // UStruct -> JSON
 
 	/**
 	 * Templated version of UStructToJsonObject to try and make most of the params. Also serves as an example use case
-	 * 
+	 *
 	 * @param InStruct The UStruct instance to read from
 	 * @param ExportCb Optional callback for types we don't understand. This is called right before falling back to the generic ToString()
 	 * @param CheckFlags Only convert properties that match at least one of these flags. If 0 check all properties.
@@ -137,7 +138,7 @@ public: // UStruct -> JSON
 	 * @return False if any properties failed to write
 	 */
 	static bool UStructToJsonAttributes(const UStruct* StructDefinition, const void* Struct, TMap< FString, TSharedPtr<FJsonValue> >& OutJsonAttributes, int64 CheckFlags, int64 SkipFlags, const CustomExportCallback* ExportCb = nullptr);
-	
+
 	/* * Converts from a UProperty to a Json Value using exportText
 	 *
 	 * @param Property			The property to export
@@ -164,7 +165,7 @@ public: // JSON -> UStruct
 	 * @return False if any properties matched but failed to deserialize
 	 */
 	static bool JsonObjectToUStruct(const TSharedRef<FJsonObject>& JsonObject, const UStruct* StructDefinition, void* OutStruct, int64 CheckFlags = 0, int64 SkipFlags = 0);
-	
+
 	/**
 	 * Templated version of JsonObjectToUStruct
 	 *
@@ -196,7 +197,7 @@ public: // JSON -> UStruct
 
 	/**
 	 * Converts a single JsonValue to the corresponding UProperty (this may recurse if the property is a UStruct for instance).
-	 * 
+	 *
 	 * @param JsonValue The value to assign to this property
 	 * @param Property The UProperty definition of the property we're setting.
 	 * @param OutValue Pointer to the property instance to be modified.
@@ -293,4 +294,10 @@ public: // JSON -> UStruct
 		}
 		return true;
 	}
+
+	/*
+	* Parses text arguments from Json into a map
+	* @param JsonObject Object to parse arguments from
+	*/
+	static FFormatNamedArguments ParseTextArgumentsFromJson(const TSharedPtr<const FJsonObject>& JsonObject);
 };
diff --git a/Engine/Source/Runtime/JsonUtilities/Public/JsonObjectWrapper.h b/Engine/Source/Runtime/JsonUtilities/Public/JsonObjectWrapper.h
--- a/Engine/Source/Runtime/JsonUtilities/Public/JsonObjectWrapper.h
+++ b/Engine/Source/Runtime/JsonUtilities/Public/JsonObjectWrapper.h
@@ -18,7 +18,7 @@ struct JSONUTILITIES_API FJsonObjectWrapper
 	GENERATED_USTRUCT_BODY()
 public:
 
-	UPROPERTY(VisibleAnywhere, Category = "JSON")
+	UPROPERTY(EditAnywhere, Category = "JSON")
 	FString JsonString;
 
 	TSharedPtr<FJsonObject> JsonObject;
diff --git a/Engine/Source/Runtime/Launch/Private/Android/AndroidJNI.cpp b/Engine/Source/Runtime/Launch/Private/Android/AndroidJNI.cpp
--- a/Engine/Source/Runtime/Launch/Private/Android/AndroidJNI.cpp
+++ b/Engine/Source/Runtime/Launch/Private/Android/AndroidJNI.cpp
@@ -120,6 +120,8 @@ void FJavaWrapper::FindClassesAndMethods(JNIEnv* Env)
 	// SurfaceView functionality for view scaling on some devices
 	AndroidThunkJava_UseSurfaceViewWorkaround = FindMethod(Env, GameActivityClassID, "AndroidThunkJava_UseSurfaceViewWorkaround", "()V", bIsOptional);
 	AndroidThunkJava_SetDesiredViewSize = FindMethod(Env, GameActivityClassID, "AndroidThunkJava_SetDesiredViewSize", "(II)V", bIsOptional);
+
+	AndroidThunkJava_IsVirtuaInputClicked = FindMethod(Env, GameActivityClassID, "AndroidThunkJava_IsVirtuaInputClicked", "(II)Z", bIsOptional);
 }
 
 void FJavaWrapper::FindGooglePlayMethods(JNIEnv* Env)
@@ -318,6 +320,8 @@ jmethodID FJavaWrapper::AndroidThunkJava_IapConsumePurchase;
 jmethodID FJavaWrapper::AndroidThunkJava_UseSurfaceViewWorkaround;
 jmethodID FJavaWrapper::AndroidThunkJava_SetDesiredViewSize;
 
+jmethodID FJavaWrapper::AndroidThunkJava_IsVirtuaInputClicked;
+
 jclass FJavaWrapper::LaunchNotificationClass;
 jfieldID FJavaWrapper::LaunchNotificationUsed;
 jfieldID FJavaWrapper::LaunchNotificationEvent;
@@ -429,6 +433,16 @@ bool AndroidThunkCpp_GetInputDeviceInfo(int32 deviceId, FAndroidInputDeviceInfo
 	return false;
 }
 
+bool AndroidThunkCpp_IsVirtuaInputClicked(int32 x, int32 y)
+{
+	bool Result = false;
+	if (JNIEnv* Env = FAndroidApplication::GetJavaEnv())
+	{
+		Result = FJavaWrapper::CallBooleanMethod(Env, FJavaWrapper::GameActivityThis, FJavaWrapper::AndroidThunkJava_IsVirtuaInputClicked, x, y);
+	}
+	return Result;
+}
+
 bool AndroidThunkCpp_IsGamepadAttached()
 {
 	bool Result = false;
diff --git a/Engine/Source/Runtime/Launch/Private/Android/LaunchAndroid.cpp b/Engine/Source/Runtime/Launch/Private/Android/LaunchAndroid.cpp
--- a/Engine/Source/Runtime/Launch/Private/Android/LaunchAndroid.cpp
+++ b/Engine/Source/Runtime/Launch/Private/Android/LaunchAndroid.cpp
@@ -27,6 +27,7 @@
 #include "HAL/PlatformAffinity.h"
 #include "Modules/ModuleManager.h"
 #include "IMessagingModule.h"
+#include "AndroidStats.h"
 
 // Function pointer for retrieving joystick events
 // Function has been part of the OS since Honeycomb, but only appeared in the
@@ -94,6 +95,7 @@ extern "C"
 
 extern void AndroidThunkCpp_InitHMDs();
 extern void AndroidThunkCpp_ShowConsoleWindow();
+extern bool AndroidThunkCpp_IsVirtuaInputClicked(int, int);
 
 // Base path for file accesses
 extern FString GFilePathBase;
@@ -135,6 +137,9 @@ static const float SampleDecayRate = 0.85f;
 // Event for coordinating pausing of the main and event handling threads to prevent background spinning
 static FEvent* EventHandlerEvent = NULL;
 
+// virtualKeyboard shown
+static volatile bool GVirtualKeyboardShown = false;
+
 // Wait for Java onCreate to complete before resume main init
 static volatile bool GResumeMainInit = false;
 volatile bool GEventHandlerInitialized = false;
@@ -330,6 +335,8 @@ int32 AndroidMain(struct android_app* state)
 	// tick until done
 	while (!GIsRequestingExit)
 	{
+		FAndroidStats::UpdateAndroidStats();
+
 		FAppEventManager::GetInstance()->Tick();
 		if(!FAppEventManager::GetInstance()->IsGamePaused())
 		{
@@ -571,6 +578,16 @@ static int32_t HandleInputCB(struct android_app* app, AInputEvent* event)
 			}
 			FPlatformRect ScreenRect = FAndroidWindow::GetScreenRect();
 
+			if (GVirtualKeyboardShown && (type == TouchBegan || type == TouchMoved))
+			{
+				int pointerId = AMotionEvent_getPointerId(event, actionPointer);
+				int32 x = AMotionEvent_getX(event, actionPointer);
+				int32 y = AMotionEvent_getY(event, actionPointer);
+
+				//ignore key down events when the native input was clicked
+				if (AndroidThunkCpp_IsVirtuaInputClicked(x, y))
+					return 0;
+			}
 			if(isActionTargeted)
 			{
 				if(actionPointer < 0 || pointerCount < (int)actionPointer)
@@ -1075,6 +1092,12 @@ static int HandleSensorEvents(int fd, int events, void* data)
 
 //Native-defined functions
 
+//Set GVirtualKeyboardShown.This function is declared in the Java-defined class, GameActivity.java: "public native void nativeVirtualKeyboardVisible(boolean bShown)"
+JNI_METHOD void Java_com_epicgames_ue4_GameActivity_nativeVirtualKeyboardVisible(JNIEnv* jenv, jobject thiz, jboolean bShown)
+{
+	GVirtualKeyboardShown = bShown;
+}
+
 //This function is declared in the Java-defined class, GameActivity.java: "public native void nativeConsoleCommand(String commandString);"
 JNI_METHOD void Java_com_epicgames_ue4_GameActivity_nativeConsoleCommand(JNIEnv* jenv, jobject thiz, jstring commandString)
 {
diff --git a/Engine/Source/Runtime/Launch/Public/Android/AndroidJNI.h b/Engine/Source/Runtime/Launch/Public/Android/AndroidJNI.h
--- a/Engine/Source/Runtime/Launch/Public/Android/AndroidJNI.h
+++ b/Engine/Source/Runtime/Launch/Public/Android/AndroidJNI.h
@@ -48,6 +48,8 @@ public:
 	static jmethodID AndroidThunkJava_HasActiveWiFiConnection;
 	static jmethodID AndroidThunkJava_GetAndroidId;
 
+	static jmethodID AndroidThunkCpp_IsVirtuaInputClicked;
+
 	// InputDeviceInfo member field ids
 	static jclass InputDeviceInfoClass;
 	static jfieldID InputDeviceInfo_VendorId;
@@ -88,6 +90,7 @@ public:
 	// SurfaceView functionality for view scaling on some devices
 	static jmethodID AndroidThunkJava_UseSurfaceViewWorkaround;
 	static jmethodID AndroidThunkJava_SetDesiredViewSize;
+	static jmethodID AndroidThunkJava_IsVirtuaInputClicked;
 
 	// member fields for getting the launch notification
 	static jclass LaunchNotificationClass;
diff --git a/Engine/Source/Runtime/PakFile/Public/IPlatformFilePak.h b/Engine/Source/Runtime/PakFile/Public/IPlatformFilePak.h
--- a/Engine/Source/Runtime/PakFile/Public/IPlatformFilePak.h
+++ b/Engine/Source/Runtime/PakFile/Public/IPlatformFilePak.h
@@ -1607,6 +1607,73 @@ public:
 		return IPlatformFile::IterateDirectoryStatRecursively(Directory, PakVisitor);
 	}
 
+	virtual void FindFiles(TArray<FString>& FoundFiles, const TCHAR* Directory, const TCHAR* FileExtension) override
+	{		
+		if (LowerLevel->DirectoryExists(Directory))
+		{
+			LowerLevel->FindFiles(FoundFiles, Directory, FileExtension);
+		}
+
+		bool bRecursive = false;
+		FindFilesInternal(FoundFiles, Directory, FileExtension, bRecursive);
+	}
+	
+	virtual void FindFilesRecursively(TArray<FString>& FoundFiles, const TCHAR* Directory, const TCHAR* FileExtension) override
+	{
+		if (LowerLevel->DirectoryExists(Directory))
+		{
+			LowerLevel->FindFiles(FoundFiles, Directory, FileExtension);
+		}
+		
+		bool bRecursive = true;
+		FindFilesInternal(FoundFiles, Directory, FileExtension, bRecursive);
+	}
+
+	void FindFilesInternal(TArray<FString>& FoundFiles, const TCHAR* Directory, const TCHAR* FileExtension, bool bRecursive)
+	{
+		TArray<FPakListEntry> Paks;
+		GetMountedPaks(Paks);
+		if (Paks.Num())
+		{
+			TSet<FString> FilesVisited;
+			FilesVisited.Append(FoundFiles);
+			
+			FString StandardDirectory = Directory;
+			FString FileExtensionStr = FileExtension;
+			FPaths::MakeStandardFilename(StandardDirectory);
+			bool bIncludeFiles = true;
+			bool bIncludeFolders = false;
+
+			TArray<FString> FilesInPak;
+			FilesInPak.Reserve(64);
+			for (int32 PakIndex = 0; PakIndex < Paks.Num(); PakIndex++)
+			{
+				FPakFile& PakFile = *Paks[PakIndex].PakFile;
+				PakFile.FindFilesAtPath(FilesInPak, *StandardDirectory, bIncludeFiles, bIncludeFolders, bRecursive);
+			}
+			
+			for (const FString& Filename : FilesInPak)
+			{
+				// filter out files by FileExtension
+				if (FileExtensionStr.Len())
+				{
+					if (!Filename.EndsWith(FileExtensionStr))
+					{
+						continue;
+					}
+				}
+								
+				// make sure we don't add duplicates to FoundFiles
+				bool bVisited = false;
+				FilesVisited.Add(Filename, &bVisited);
+				if (!bVisited)
+				{
+					FoundFiles.Add(Filename);
+				}
+			}
+		}
+	}
+	
 	virtual bool DeleteDirectoryRecursively(const TCHAR* Directory) override
 	{
 		// Can't delete directories existing in pak files. See DeleteDirectory(..) for more info.
diff --git a/Engine/Source/Runtime/Projects/Private/PluginManager.cpp b/Engine/Source/Runtime/Projects/Private/PluginManager.cpp
--- a/Engine/Source/Runtime/Projects/Private/PluginManager.cpp
+++ b/Engine/Source/Runtime/Projects/Private/PluginManager.cpp
@@ -365,49 +365,7 @@ void FPluginManager::ReadPluginsInDirectory(const FString& PluginsDirectory, con
 
 void FPluginManager::FindPluginsInDirectory(const FString& PluginsDirectory, TArray<FString>& FileNames)
 {
-	// Class to enumerate the contents of a directory, and find all sub-directories and plugin descriptors within it
-	struct FPluginDirectoryVisitor : public IPlatformFile::FDirectoryVisitor
-	{
-		TArray<FString> SubDirectories;
-		TArray<FString> PluginDescriptors;
-
-		virtual bool Visit(const TCHAR* FilenameOrDirectory, bool bIsDirectory) override
-		{
-			FString FilenameOrDirectoryStr = FilenameOrDirectory;
-			if(bIsDirectory)
-			{
-				SubDirectories.Add(FilenameOrDirectoryStr);
-			}
-			else if(FilenameOrDirectoryStr.EndsWith(TEXT(".uplugin")))
-			{
-				PluginDescriptors.Add(FilenameOrDirectoryStr);
-			}
-			return true;
-		}
-	};
-
-	// Enumerate the contents of the current directory
-	FPluginDirectoryVisitor Visitor;
-	FPlatformFileManager::Get().GetPlatformFile().IterateDirectory(*PluginsDirectory, Visitor);
-
-	// If there's no plugins in this directory, recurse through all the child directories
-	if(Visitor.PluginDescriptors.Num() == 0)
-	{
-		for(const FString& SubDirectory: Visitor.SubDirectories)
-		{
-			FindPluginsInDirectory(SubDirectory, FileNames);
-		}
-	}
-	else
-	{
-		for(const FString& PluginDescriptor: Visitor.PluginDescriptors)
-		{
-			if (!FileNames.Contains(PluginDescriptor))
-			{
-				FileNames.Add(PluginDescriptor);
-			}
-		}
-	}
+	FPlatformFileManager::Get().GetPlatformFile().FindFilesRecursively(FileNames, *PluginsDirectory, TEXT(".uplugin"));
 }
 
 void FPluginManager::FindPluginManifestsInDirectory(const FString& PluginManifestDirectory, TArray<FString>& FileNames)
diff --git a/Engine/Source/Runtime/SandboxFile/Public/IPlatformFileSandboxWrapper.h b/Engine/Source/Runtime/SandboxFile/Public/IPlatformFileSandboxWrapper.h
--- a/Engine/Source/Runtime/SandboxFile/Public/IPlatformFileSandboxWrapper.h
+++ b/Engine/Source/Runtime/SandboxFile/Public/IPlatformFileSandboxWrapper.h
@@ -108,6 +108,9 @@ class SANDBOXFILE_API FSandboxPlatformFile : public IPlatformFile
 	 * @param Files true to include directories in the Result
 	 */
 	void FindFiles( TArray<FString>& Result, const TCHAR* InFilename, bool Files, bool Directories );
+	
+	/** Allow IPlatformFile::FindFiles */
+	using IPlatformFile::FindFiles;
 
 	/**
 	 * Deletes a directory
diff --git a/Engine/Source/Runtime/Slate/Private/Framework/Application/SlateApplication.cpp b/Engine/Source/Runtime/Slate/Private/Framework/Application/SlateApplication.cpp
--- a/Engine/Source/Runtime/Slate/Private/Framework/Application/SlateApplication.cpp
+++ b/Engine/Source/Runtime/Slate/Private/Framework/Application/SlateApplication.cpp
@@ -5106,6 +5106,10 @@ bool FSlateApplication::ProcessMouseButtonDownEvent( const TSharedPtr< FGenericW
 FReply FSlateApplication::RoutePointerDownEvent(FWidgetPath& WidgetsUnderPointer, FPointerEvent& PointerEvent)
 {
 	TScopeCounter<int32> BeginInput(ProcessingInput);
+	
+	// Ensure the cursor location(s) get set to an initial value
+	PointerIndexPositionMap.Add(FUserAndPointer(PointerEvent.GetUserIndex(), PointerEvent.GetPointerIndex()), PointerEvent.GetScreenSpacePosition());
+	PointerIndexLastPositionMap.Add(FUserAndPointer(PointerEvent.GetUserIndex(), PointerEvent.GetPointerIndex()), PointerEvent.GetScreenSpacePosition());
 
 #if PLATFORM_MAC
 	NSWindow* ActiveWindow = [ NSApp keyWindow ];
diff --git a/Engine/Source/Runtime/UMG/Private/Components/ScaleBox.cpp b/Engine/Source/Runtime/UMG/Private/Components/ScaleBox.cpp
--- a/Engine/Source/Runtime/UMG/Private/Components/ScaleBox.cpp
+++ b/Engine/Source/Runtime/UMG/Private/Components/ScaleBox.cpp
@@ -19,7 +19,9 @@ UScaleBox::UScaleBox(const FObjectInitializer& ObjectInitializer)
 	Stretch = EStretch::ScaleToFit;
 	UserSpecifiedScale = 1.0f;
 	IgnoreInheritedScale = false;
-	bSingleLayoutPass = false;
+	// @todo: - DO NOT copy this change to UE4/Main
+	bSingleLayoutPass = true;
+	// @todo: - DO NOT copy this change to UE4/Main
 }
 
 void UScaleBox::ReleaseSlateResources(bool bReleaseChildren)
diff --git a/Engine/Source/Runtime/UMG/Private/Slate/SWorldWidgetScreenLayer.cpp b/Engine/Source/Runtime/UMG/Private/Slate/SWorldWidgetScreenLayer.cpp
--- a/Engine/Source/Runtime/UMG/Private/Slate/SWorldWidgetScreenLayer.cpp
+++ b/Engine/Source/Runtime/UMG/Private/Slate/SWorldWidgetScreenLayer.cpp
@@ -107,7 +107,9 @@ void SWorldWidgetScreenLayer::Tick(const FGeometry& AllottedGeometry, const doub
 								CanvasSlot->Offset(FMargin(LocalPosition.X, LocalPosition.Y, ComponentDrawSize.X, ComponentDrawSize.Y));
 								CanvasSlot->Anchors(FAnchors(0, 0, 0, 0));
 								CanvasSlot->Alignment(ComponentPivot);
-								CanvasSlot->ZOrder(-ViewportPosition.Z);
+								// @todo: - DO NOT copy this change to UE4/Main
+								//CanvasSlot->ZOrder(-ViewportPosition.Z);
+								// @todo: - DO NOT copy this change to UE4/Main
 							}
 							else
 							{
@@ -115,7 +117,9 @@ void SWorldWidgetScreenLayer::Tick(const FGeometry& AllottedGeometry, const doub
 								CanvasSlot->Offset(FMargin(LocalPosition.X, LocalPosition.Y, DrawSize.X, DrawSize.Y));
 								CanvasSlot->Anchors(FAnchors(0, 0, 0, 0));
 								CanvasSlot->Alignment(Pivot);
-								CanvasSlot->ZOrder(-ViewportPosition.Z);
+								// @todo: - DO NOT copy this change to UE4/Main
+								//CanvasSlot->ZOrder(-ViewportPosition.Z);
+								// @todo: - DO NOT copy this change to UE4/Main
 							}
 						}
 					}
