Copying //UE4/Dev-Sequencer to //UE4/Main (Source: //UE4/Dev-Sequencer @ 3617642)

Change 3235667 on 2016/12/14 by [REDACTED_PERSON]

	Media: Moved enums into separate header file, so they can be shared

Change 3259266 on 2017/01/16 by [REDACTED_PERSON]

	Core: Added timespan ratio

Change 3267229 on 2017/01/21 by [REDACTED_PERSON]

	Editor: Fixed texture samplers in materials not getting notified when sRGB mode changed

Change 3274773 on 2017/01/27 by [REDACTED_PERSON]

	Core: Added TLruCache template

Change 3281579 on 2017/02/01 by [REDACTED_PERSON]

	Core: Added scalar division to FTimespan

Change 3289522 on 2017/02/07 by [REDACTED_PERSON]

	MediaAssets: Added looping for play lists

Change 3290664 on 2017/02/07 by [REDACTED_PERSON]

	Engine: Moved UTexture material notification code into separate function, so it can be used by derived classes.

Change 3290688 on 2017/02/07 by [REDACTED_PERSON]

	MediaAssets: Notifying materials when media texture properties changed.

Change 3291171 on 2017/02/07 by [REDACTED_PERSON]

	MediaAssets: Filtering unused UTexture properties

Change 3291229 on 2017/02/07 by [REDACTED_PERSON]

	ImgMedia: Use displayWindow instead of dataWindow due to bug in OpenEXR with data window min < 0

Change 3298520 on 2017/02/11 by [REDACTED_PERSON]

	Sequencer: Fixed skylight and reflection components off by one frame in PIE

Change 3298778 on 2017/02/12 by [REDACTED_PERSON]

	MediaUtils: Experimenting with dynamic pitch adjustment

Change 3298987 on 2017/02/13 by [REDACTED_PERSON]

	Editor: Add matrix and transform properties to property editor test object

Change 3298997 on 2017/02/13 by [REDACTED_PERSON]

	Editor: It's now possible to specify Units and ForceUnits meta-data on struct properties
	  - Doing so will cause any child numeric and struct properties to use these units by default (unless another unit was specified)
	  - Added percentage unit type
	  - Multiplier unit types can now be converted between

Change 3298998 on 2017/02/13 by [REDACTED_PERSON]

	Editor: Conversion between multiplier (1.5x)  and percentage (150%) units is now possible

Change 3348678 on 2017/03/15 by [REDACTED_PERSON]

	Media: Added Buffering media event

Change 3355268 on 2017/03/20 by [REDACTED_PERSON]

	Core: TQueue documentation updates

Change 3359055 on 2017/03/22 by [REDACTED_PERSON]

	Portal: Removed dependency to Messaging.h

Change 3359060 on 2017/03/22 by [REDACTED_PERSON]

	Messaging: Modernization pass

	- added missing includes & forward declarations
	- include what you use
	- deprecated shared pointer typedefs
	- replaced some delegates with callback interfaces
	- documentation fixes

Change 3359189 on 2017/03/22 by [REDACTED_PERSON]

	Automation: Modernization pass

	- include what you use
	- removed molothic header
	- documentation fixes

	#upgradenotes: IAutomationWorkerModule.h is now in the module's public root directory

Change 3359718 on 2017/03/22 by [REDACTED_PERSON]

	Messaging: Moved common helper classes into MessagingCommon module

	#upgradenotes:  If you use the common helper classes, replace "Messaging" with "MessagingCommon" in your Build.cs files and remove the "Helpers/" subdirectory from your include statements (unless you actually
include header files from "Messaging", in which case you need to keep "Messaging" and add "MessagingCommon").

Change 3359793 on 2017/03/23 by [REDACTED_PERSON]

	TargetDeviceServices: Modernization pass

	- include what you use
	- removed boilerplate header
	- deprecated selected shared pointer typedefs
	- reorganized files
	- documentation fixes

Change 3361028 on 2017/03/23 by [REDACTED_PERSON]

	DeviceManager: Modernization pass

	- include what you use
	- reduced shared pointer typedef usage
	- documentation fixes

Change 3361197 on 2017/03/23 by [REDACTED_PERSON]

	Messaging: Waking up message router thread on shutdown

Change 3361246 on 2017/03/23 by [REDACTED_PERSON]

	AutomationWindow: Removed boilerplate header

	#upgradenotes: Replace includes of AutomationWindow.h with IAutomationWindowModule.h

Change 3361428 on 2017/03/23 by [REDACTED_PERSON]

	AutomationController: Removed boilerplate header

	#upgradenotes: Instead of AutomationController.h, include the individual interface headers that you actually use. Remove the "Interfaces/" subdirectory from existing interface inclusions.

Change 3363206 on 2017/03/24 by [REDACTED_PERSON]

	ProfileLauncher: Modernization pass

	- include what you use (selected files)
	- reduced shared pointer typedef usages
	- removed dead code
	- cleaned up file organization
	- documentation fixes

Change 3363290 on 2017/03/24 by [REDACTED_PERSON]

	LauncherServices: Removed boilerplate header

	#upgradenotes: Instead of including LauncherServices.h, include the ILauncherXXX.h files that you actually use. Remove the "Interfaces/" subdirectory from existing ILauncherXXX.h includes

Change 3363305 on 2017/03/24 by [REDACTED_PERSON]

	LauncherCheck: Removed boilerplate header

	#upgradenotes: Include ILauncherCheckModule.h instead of LauncherCheck.h. Remove "Interfaces/" subdirectory from existing ILauncherCheckModule.h inclusions

Change 3363708 on 2017/03/24 by [REDACTED_PERSON]

	ImageWrapper: Removed boilerplate header

	#upgradenotes: Instead of including ImageWrapper.h, include the actual IImageWrapperXXX.h files that you use. Remove the "Interfaces/" subdirectory from existing IImageWrapperXXX.h includes

Change 3363966 on 2017/03/24 by [REDACTED_PERSON]

	ImageWrapper: Modernization pass

	- include what you use
	- reorganized internal files
	- documentation fixes

Change 3364579 on 2017/03/24 by [REDACTED_PERSON]

	ImageWrapper: Deprecated shared pointer typedef and removed usages

	#upgradenotes: Please use TSharedPtr<IImageWrapper> instead of IImageWrapperPtr

Change 3364582 on 2017/03/24 by [REDACTED_PERSON]

	NetworkFileSystem: Removed monolithic boilerplate header

	#upgradenotes: Instead of including NetworkFileServer.h, include the INetworkFileXXX.h headers that you actually use. Remove the "Interfaces/" subdirectory from existing INetworkFileXXX.h includes

Change 3381440 on 2017/04/05 by [REDACTED_PERSON]

	Oculus: Removed illegal pragmas

Change 3391731 on 2017/04/12 by [REDACTED_PERSON]

	ImgMedia: Added support for BMP, JPG and PNG image sequences

Change 3401146 on 2017/04/19 by [REDACTED_PERSON]

	Core: Cleanup pass for FTimespan usages; updated documentation

	When initializing time span values from single components, consider using the FromHours, FromMinutes, FromSeconds, Zero, MinValue and related methods instead of calling the overloaded constructors as they will make
your code easier to read and understand.

Change 3401504 on 2017/04/20 by [REDACTED_PERSON]

	MCP: Fixed incorrect FTimespan usage of GetMilliseconds instead of GetTotalMilliseconds

Change 3401833 on 2017/04/20 by [REDACTED_PERSON]

	Core: Various improvements to FTimespan

	- added proper serialization
	- updated documentation
	- access fractions as milli-, micro- or nanoseconds
	- removed the following string formatters as they were not useful: %D %H %M %S %F
	- updated documentation

	#upgradenotes: The %D string formatter (previously total number of days) is now the zero-padded days component of the time span. Time spans are now always exported and stringified with a leading plus or minus sign.

Change 3437543 on 2017/05/12 by [REDACTED_PERSON]

	PS4Media: Added audio error codes conversion

Change 3446564 on 2017/05/18 by [REDACTED_PERSON]

	Networking: Fixed nullpointer crash if socket initialization fails

Change 3459978 on 2017/05/25 by [REDACTED_PERSON]

	ImgMedia: Renamed frame cache to frame loader

Change 3463311 on 2017/05/26 by [REDACTED_PERSON]

	Core: Added iterators & predicate based methods to LRU Cache

Change 3464452 on 2017/05/29 by [REDACTED_PERSON]

	ImgMedia: Added video sample output

Change 3464468 on 2017/05/29 by [REDACTED_PERSON]

	Media: Added separate cache visualization for loaded and cached samples

Change 3464592 on 2017/05/29 by [REDACTED_PERSON]

	Core: Added TLruCache::FindAndTouch

Change 3464607 on 2017/05/29 by [REDACTED_PERSON]

	Core: Added TRange::Inclusive / ::Exclusive

Change 3464608 on 2017/05/29 by [REDACTED_PERSON]

	Media: Fixed incorrect upper bound in support play rates of various players

Change 3466732 on 2017/05/30 by [REDACTED_PERSON]

	ImgMedia: Added image compression type to info string

Change 3466871 on 2017/05/31 by [REDACTED_PERSON]

	WmfMedia: Fixed session capabilities not showing up; added initialization logging

Change 3467785 on 2017/05/31 by [REDACTED_PERSON]

	MfMedia: Fixed sample stride calculated from media input instead of output

Change 3467963 on 2017/05/31 by [REDACTED_PERSON]

	WmfMedia: Reordered media sub-type string conversion to fix some tracks being reported incorrectly

Change 3468110 on 2017/05/31 by [REDACTED_PERSON]

	WmfMedia: Added MPEG-2 media sub types to utils

Change 3468516 on 2017/05/31 by [REDACTED_PERSON]

	AndroidMedia: Allow mediaplayer audio to be disable on Android (UE-45570)

Change 3468554 on 2017/05/31 by [REDACTED_PERSON]

	MediaAssets: Added AutoClear option to automatically clear media texture when media is unloaded

Change 3468627 on 2017/05/31 by [REDACTED_PERSON]

	ImgMedia: Allowing for auto-selection of EXR decoder thread count

Change 3468648 on 2017/05/31 by [REDACTED_PERSON]

	ImgMedia: Reorganized OpenExrWrapper files

Change 3471789 on 2017/06/02 by [REDACTED_PERSON]

	WmfMedia: Fixed potential concurrency issue in session state management

Change 3472918 on 2017/06/03 by [REDACTED_PERSON]

	Core: Allowing zero sized LRU cache; added default constructor

Change 3472919 on 2017/06/03 by [REDACTED_PERSON]

	Media: Added pre-Slate tick stage; renamed some functions; continued to implement sample caching

Change 3473704 on 2017/06/05 by [REDACTED_PERSON]

	AndroidMedia: Added settings class

Change 3474407 on 2017/06/05 by [REDACTED_PERSON]

	AndroidMedia: Implemented video sample processing

Change 3474934 on 2017/06/05 by [REDACTED_PERSON]

	ImgMedia: Fixed crash due to deleting loader work items twice (UE-45705)

Change 3476925 on 2017/06/06 by [REDACTED_PERSON]

	Media: Added return values to media view interface; documentation fixes

Change 3477158 on 2017/06/06 by [REDACTED_PERSON]

	SwitchMedia: Implemented sample pooling

Change 3477201 on 2017/06/07 by [REDACTED_PERSON]

	PS4Media: Implemented video sample pooling

Change 3481470 on 2017/06/08 by [REDACTED_PERSON]

	Media: Fixed time stamps not set in recycled texture samples

Change 3481472 on 2017/06/08 by [REDACTED_PERSON]

	MediaAssets: Fixed no video samples fetched while player is in buffering state

Change 3481473 on 2017/06/08 by [REDACTED_PERSON]

	MediaUtils: Added default constructor to MediaSampleQueue

Change 3481584 on 2017/06/08 by [REDACTED_PERSON]

	MediaUtils: Paused state is considered forward for sample lookup as this is the common case.

Change 3481588 on 2017/06/08 by [REDACTED_PERSON]

	PS4Media: Redesigned the PS4 media player to support async file loading & track switching

Change 3481627 on 2017/06/08 by [REDACTED_PERSON]

	PS4Media: Opening media sources in thread pool if not precaching

Change 3481666 on 2017/06/08 by [REDACTED_PERSON]

	ImgMedia: Fixed SequencePath property file picker always opening default path

Change 3481669 on 2017/06/08 by [REDACTED_PERSON]

	ImgMedia: Fixed crash when quering cache state of empty image sequence

Change 3481685 on 2017/06/08 by [REDACTED_PERSON]

	ImgMedia: Fixed EXR reader failing to play slow loading sequences (UE-45719)

Change 3483623 on 2017/06/10 by [REDACTED_PERSON]

	Media: Added getters for 360 view settings; added support for relative view updates

Change 3483624 on 2017/06/10 by [REDACTED_PERSON]

	Media: Added Blueprint support for 360 videos

Change 3483626 on 2017/06/10 by [REDACTED_PERSON]

	MediaPlayerEditor: Added 360 video mouse controls; broke out viewport widget into separate class.

Change 3483627 on 2017/06/10 by [REDACTED_PERSON]

	Core: Using system start time instead of current time for renamed log file names

Change 3483630 on 2017/06/10 by [REDACTED_PERSON]

	Core: Grouping log files by log name and processing each group separately when deleting old log files

Change 3483816 on 2017/06/10 by [REDACTED_PERSON]

	WmfMedia: Detecting audio device availability to prevent lockup of audio tracks

Change 3483939 on 2017/06/11 by [REDACTED_PERSON]

	AvfMedia: Finished Media Framework 3.0 upgrade

	- overhauled track switching
	- removed render thread flushes
	- reduced nesting
	- code cleanup pass

Change 3483940 on 2017/06/11 by [REDACTED_PERSON]

	Media: Consistent track switching behavior across platforms

Change 3484172 on 2017/06/11 by [REDACTED_PERSON]

	MediaPlayerEditor: Removed obsolete asset type actions

Change 3484180 on 2017/06/11 by [REDACTED_PERSON]

	MediaPlayerEditor: Fixed media player info overlays not visible during mouse capture

Change 3484248 on 2017/06/11 by [REDACTED_PERSON]

	MediaAssets: Media players now always have a playlist

Change 3484249 on 2017/06/11 by [REDACTED_PERSON]

	MediaPlayerEditor: Added UI controls to Playlist tab

Change 3484250 on 2017/06/11 by [REDACTED_PERSON]

	Media: Implemented overlay sample caching

Change 3484252 on 2017/06/11 by [REDACTED_PERSON]

	WmfMedia: Enabled support for .smi and .sami subtitle files

Change 3485433 on 2017/06/12 by [REDACTED_PERSON]

	ImgMedia: Enabled plug-in on all platforms; enabled EXR only on platforms that support it

Change 3485720 on 2017/06/12 by [REDACTED_PERSON]

	MediaPlayerEditor: Added ability to save playlists

Change 3485828 on 2017/06/12 by [REDACTED_PERSON]

	AvfMedia: Thread-safe ticking of video sampler; ticking audio processing on high frequency thread

Change 3485926 on 2017/06/12 by [REDACTED_PERSON]

	MediaUtils: Draining all unconsumed samples at the end of a frame

Change 3486043 on 2017/06/12 by [REDACTED_PERSON]

	Media: Consistent behavior for playlist navigation (UE-45964)

Change 3486104 on 2017/06/12 by [REDACTED_PERSON]

	MediaPlayerEditor: Implemented simple overlay text positioning

Change 3486145 on 2017/06/12 by [REDACTED_PERSON]

	AndroidMediaPlayer: Moved video sample handling into render thread

Change 3486147 on 2017/06/12 by [REDACTED_PERSON]

	MediaPlayerEditor: Moved overlay texts into separate layer

Change 3486188 on 2017/06/12 by [REDACTED_PERSON]

	Media: Enabling media factory modules in Editor for all players

Change 3486223 on 2017/06/12 by [REDACTED_PERSON]

	Media: Defaulting players to select first audio and video tracks by default

Change 3486473 on 2017/06/13 by [REDACTED_PERSON]

	Media: Buffering is no longer a media player state, but handled separately (for players that can buffer while playing)

Change 3486475 on 2017/06/13 by [REDACTED_PERSON]

	MediaPlayerEditor: Showing busy bar when buffering or preparing

Change 3487237 on 2017/06/13 by [REDACTED_PERSON]

	AndroidMedia: Fixed video sample not getting processed on render thread

Change 3487507 on 2017/06/13 by [REDACTED_PERSON]

	AvfMedia - small cleanup of the code for passing video frame rate to the video sampler

Change 3487719 on 2017/06/13 by [REDACTED_PERSON]

	Disable shared PCH in OpenExrWrapper to solve a problem with bUseRTTI mismatch between the module and the PCH

Change 3487842 on 2017/06/13 by [REDACTED_PERSON]

	AndroidMedia: Fixed texture swizzle disabled on non-Engine builds

Change 3488006 on 2017/06/13 by [REDACTED_PERSON]

	AvfMedia - Moved some logic from FAvfMediaPlayer::TickTickable() to FAvfMediaTracks::ProcessAudio() to make it thread safe, plus fixed a bug with the player not setting the current status to Stopped after opening a
file

Change 3488308 on 2017/06/13 by [REDACTED_PERSON]

	AndroidMedia: Fixed flicker issue

Change 3488335 on 2017/06/13 by [REDACTED_PERSON]

	MediaAssets: Made UFileMediaSource::GetFullPath public

Change 3488338 on 2017/06/13 by [REDACTED_PERSON]

	MediaPlayerEditor: Fixed "Show file in Explorer" not working for file media sources

Change 3488339 on 2017/06/13 by [REDACTED_PERSON]

	MediaPlayerEditor: Showing a throbber when buffering player

Change 3488768 on 2017/06/14 by [REDACTED_PERSON]

	Core: Inlined FTimespan static functions

hange 3490203 on 2017/06/14 by [REDACTED_PERSON]

	Core: Force inlined TComPtr operators

Change 3494083 on 2017/06/15 by [REDACTED_PERSON]

	AndroidMedia: Fix playlist looping (send PlaybackEndReached event)

Change 3497017 on 2017/06/16 by [REDACTED_PERSON]

	MediaAssets: Playlists now auto advance with PlayOnOpen off

Change 3497075 on 2017/06/17 by [REDACTED_PERSON]

	MediaPlayer Sampler node and external texture support for Android MediaPlayer
	- enable with USE_EXTERNALTEXTURE in AndroidMediaPlayer.cpp (disabled in this CL)
	- supports MediaSample node in material editor with scale/bias handled if above flag enabled
	- fixed looping problems, including eventual crashes
	- track switching fixed (lockups and missing video)
	- corrected resource leaks and crash switching maps

Change 3497163 on 2017/06/17 by [REDACTED_PERSON]

	MediaUtils: Processing media events immediately if on game thread

Change 3497170 on 2017/06/17 by [REDACTED_PERSON]

	Media: All Media Framework interfaces are pure virtual

Change 3498603 on 2017/06/19 by [REDACTED_PERSON]

	UBT: Prevent plugins which list modules multiple times from adding them twice.

Change 3500870 on 2017/06/20 by [REDACTED_PERSON]

	WmfMedia: Rewrote WMF state machine; now with track switching

Change 3502181 on 2017/06/20 by [REDACTED_PERSON]

	MediaSampler node updates
	- fixed issue with ExternalTexture singleton for DLL compatiblity
	- now works for all players (registers texture samples with ExternalTexture by player GUID)
	- enabled OES in AndroidMediaPlayer now that materials can universally use the MediaSampler

Change 3503182 on 2017/06/21 by [REDACTED_PERSON]

	Media: Refactored status flags into an enum; exposed connecting status

Change 3503724 on 2017/06/21 by [REDACTED_PERSON]

	Improvement in MediaPlayer Sampler node
	- RGB, R, G, B, A output pins like normal texture sampler
	- allow preview in material editor
	- fix OES detection for Android

Change 3509088 on 2017/06/26 by [REDACTED_PERSON]

	MediaAssets: Added AddFile and AddUrl to UMediaPlaylist

Change 3510256 on 2017/06/26 by [REDACTED_PERSON]

	WmfMedia: Fixed play rate not set to zero when end reached

Change 3510273 on 2017/06/26 by [REDACTED_PERSON]

	MediaPlayerEditor: Removed obsolete Slate culling rectangle code

Change 3510413 on 2017/06/26 by [REDACTED_PERSON]

	MediaPlayerEditor: Showing frame rate if available

Change 3510533 on 2017/06/26 by [REDACTED_PERSON]

	MediaUtils: Preventing media cache filling up when scrubbing

Change 3510859 on 2017/06/26 by [REDACTED_PERSON]

	PS4Media: Fixed track switching crashes (UE-45960)

Change 3514173 on 2017/06/28 by [REDACTED_PERSON]

	WmfMedia: Optimized player capabilities check

Change 3514174 on 2017/06/28 by [REDACTED_PERSON]

	WmfMedia: Moved media source resolver code into utility class

Change 3514714 on 2017/06/28 by [REDACTED_PERSON]

	Core: Added TComPointer.IsValid; code cleanup pass

Change 3517912 on 2017/06/30 by [REDACTED_PERSON]

	MediaUtils: Added GetVideoAspectRatio to player facade

Change 3524957 on 2017/07/06 by [REDACTED_PERSON]

	ImgMedia: Decoding image frames only when video track is selected

Change 3525252 on 2017/07/06 by [REDACTED_PERSON]

	ImgMedia: Async image sequence initialization

Change 3525266 on 2017/07/06 by [REDACTED_PERSON]

	ImgMedia: Enabled reverse playback

Change 3525722 on 2017/07/06 by [REDACTED_PERSON]

	WmfMedia: Workaround for occasional WMF internal deadlock in IMFRateSupport

Change 3525800 on 2017/07/07 by [REDACTED_PERSON]

	WmfMedia: Fixed another WMF deadlock issue & some rate change glitches; improved logging.

Change 3525801 on 2017/07/07 by [REDACTED_PERSON]

	MediaPlayerEditor: Made scrubbing more responsive

Change 3526500 on 2017/07/07 by [REDACTED_PERSON]

	WmfMedia: Using NULL for Windows pointers

Change 3527323 on 2017/07/07 by [REDACTED_PERSON]

	WmfMedia: Added support for audio and video capture media sources

Change 3530197 on 2017/07/10 by [REDACTED_PERSON]

	WmfMedia: Added utility functions for enumerating audio and video capture devices

Change 3533465 on 2017/07/12 by [REDACTED_PERSON]

	Media: Added media capture device support API & implementation for WMF

Change 3533469 on 2017/07/12 by [REDACTED_PERSON]

	MediaPlayerEditor: Added source selection menu w/ capture device menu to navigation bar

Change 3533540 on 2017/07/12 by [REDACTED_PERSON]

	MediaAssets: Blueprint support for media capture device discovery

Change 3533574 on 2017/07/12 by [REDACTED_PERSON]

	WmfMedia: Using cached supported play rates; more strict pausing support check

Change 3533924 on 2017/07/12 by [REDACTED_PERSON]

	WmfMedia: Added global LowLatency setting (UEVR-859)

	Note that this setting is only supported when compiling for Windows 8 or newer.

Change 3534027 on 2017/07/12 by [REDACTED_PERSON]

	WmfMedia: Compile time support for low latency session attribute on < Windows8

Change 3538744 on 2017/07/14 by [REDACTED_PERSON]

	MediaUtils: Renamed FDefaultMediaTimeSource to FAppMediaTimeSource

Change 3542818 on 2017/07/18 by [REDACTED_PERSON]

	WmfMedia: Fixed race condition on session shutdown

Change 3543082 on 2017/07/18 by [REDACTED_PERSON]

	WmfMedia: Fixed track selection in media that doesn't support seeking

Change 3543092 on 2017/07/18 by [REDACTED_PERSON]

	Media: Implemented media track format API

	Enables multiple formats per media track. By default, the first usable format is selected.
	This feature is currently only implemented for WmfMedia. All other players assume a single format per track.

Change 3543794 on 2017/07/19 by [REDACTED_PERSON]

	WmfMedia: Added string conversion for null GUIDs

Change 3543796 on 2017/07/19 by [REDACTED_PERSON]

	MfMedia: Copied GUID string conversion updates from WmfMedia

Change 3543797 on 2017/07/19 by [REDACTED_PERSON]

	WmfMedia: Preventing duplicate track formats in capture devices with legacy DirectX support

Change 3544390 on 2017/07/19 by [REDACTED_PERSON]

	Media: Allowing INDEX_NONE as 'current selection' index in track format related functions

Change 3545368 on 2017/07/19 by [REDACTED_PERSON]

	WmfMedia: Fixed session error with some media sources when attempting to seek to current position without ever starting

Change 3545388 on 2017/07/19 by [REDACTED_PERSON]

	MediaAssets: Fixed RGB input sources rendered upside down

Change 3545430 on 2017/07/19 by [REDACTED_PERSON]

	MediaAssets: Fixed incorrect sRGB conversion on BMP inputs

Change 3547362 on 2017/07/20 by [REDACTED_PERSON]

	Core: Added IsGraph and IsPrint to TCHAR utilities

Change 3547376 on 2017/07/20 by [REDACTED_PERSON]

	WmfMedia: Better log messages for non-standard video types (UE-47533)

Change 3547404 on 2017/07/20 by [REDACTED_PERSON]

	QAGame: Enabled PlayOnOpen and Looping on MediaPlayerSwitch BP (UE-47542)

Change 3547466 on 2017/07/20 by [REDACTED_PERSON]

	WmfMedia: Fixed edge case for detecting whether pause is available

Change 3548742 on 2017/07/21 by [REDACTED_PERSON]

	Media: Added API for changing the input frame rate

Change 3548743 on 2017/07/21 by [REDACTED_PERSON]

	WmfMedia: Implemented ability to change input frame rate

Change 3554411 on 2017/07/25 by [REDACTED_PERSON]

	AudioMixer: Broke out SynthComponent initialization code; added bIsUISound property

Change 3554818 on 2017/07/25 by [REDACTED_PERSON]

	Media: Added per track-type cache queries

Change 3557284 on 2017/07/26 by [REDACTED_PERSON]

	WmfMedia: Added buffer size check in texture sample as well

Change 3560530 on 2017/07/27 by [REDACTED_PERSON]

	WmfMedia: Made COM object destructors private and added assertions

Change 3560580 on 2017/07/27 by [REDACTED_PERSON]

	MediaUtils: Added method for querying number of objects in pool

Change 3562572 on 2017/07/28 by [REDACTED_PERSON]

	WmfMedia: Properly handling topology status errors when opening media

Change 3576710 on 2017/08/08 by [REDACTED_PERSON]

	bug fixes for mediaplayer
	- rare cases in Android media tracks
	- make FMediaCaptureDevice members available in blueprints

Change 3577736 on 2017/08/08 by [REDACTED_PERSON]

	Extend ExternalTexture to support 2x3 transform (scale/rotation + offset)

Change 3578831 on 2017/08/09 by [REDACTED_PERSON]

	MediaAssets: Made media sound component BP spawnable

Change 3579210 on 2017/08/09 by [REDACTED_PERSON]

	UdpMessaging: Not showing warning message if tunnel settings left empty to use defaults

Change 3579547 on 2017/08/09 by [REDACTED_PERSON]

	Better support for Android external texture detection

Change 3579567 on 2017/08/09 by [REDACTED_PERSON]

	Android camera plugin

Change 3580035 on 2017/08/10 by [REDACTED_PERSON]

	Media Texture: External textures are now fully supported and enforced when attempting to sample external texture types in materials.
	  - Textures can now dynamically provide an external texture GUID at runtime, which is queried by the uniform expression. Statically defined GUIDs remain supported.
	  - Added necessary ExternalTexture(Parameter) methods to the material compiler and HLSL translator
	  - Made UTexture::GetMaterialType const-correct

	#tests Test map that contains TextureObject, TextureObjectParameter, TextureSample and TextureSampleParameter nodes with media textures works on both a Nexus 10 (non-image-external), and a Pixel (image-external),
and on desktop PC.

Change 3581552 on 2017/08/10 by [REDACTED_PERSON]

	Use bilinear sampling instead of point for external texture (Android media and camera)

Change 3581628 on 2017/08/10 by [REDACTED_PERSON]

	Core: Fixed FTimespan import/export/copypaste (UE-43990)

Change 3581909 on 2017/08/11 by [REDACTED_PERSON]

	Media: Only cache material proxy pointers if they have been initialized (and thus, will be removed from the cache on release)
	  - This prevents destroyed proxies from remaining in the texture registry

Change 3582451 on 2017/08/11 by [REDACTED_PERSON]

	AudioMixer: Fixed initialization order of synth components

	This ensures that the Synth is initialized before OnRegister in the super class initializes and activates the AudioComponents. Previously, the AudioComponent would fail to activate, because Synth was nullptr.

Change 3582453 on 2017/08/11 by [REDACTED_PERSON]

	MediaAssets: Enabling ticking in media sound component (UE-48055)

Change 3583101 on 2017/08/11 by [REDACTED_PERSON]

	Android camera improvements
	- return correct capture device type (webcamfront/rear)
	- remove seek (not supported)
	- support track formats
	- better framerate selection

Change 3590109 on 2017/08/16 by [REDACTED_PERSON]

	Move external texture coordinate update to improve accuracy and save a uniform in final shader

Change 3590530 on 2017/08/16 by [REDACTED_PERSON]

	MediaAssets: Changed capture device enumeration BP function signatures to allow filter bit masks

Change 3590547 on 2017/08/16 by [REDACTED_PERSON]

	MediaPlayerEditor: Ticking sound component directly

Change 3590628 on 2017/08/16 by [REDACTED_PERSON]

	Switch: Fixed media decoder asserting when attempting to pause without having started

Change 3591816 on 2017/08/16 by [REDACTED_PERSON]

Change 3592266 on 2017/08/16 by [REDACTED_PERSON]

	Media: Replaced media cache with the new sample processing pipeline; continued to implement MfMedia

	This change fixes stuttering issues with audio and video. However, proper pipeline flushing is not fully implemented yet, so the playback will break after the first playthrough, which might result in freeze frames
or out of memory assertions depending on the platform. This will be addressed in an upcoming change.

Change 3592355 on 2017/08/17 by [REDACTED_PERSON]

	MediaUtils: Added sample sink collection to remove code duplication

Change 3592739 on 2017/08/17 by [REDACTED_PERSON]

	ImgMedia: Fetching only one video sample per frame

Change 3592741 on 2017/08/17 by [REDACTED_PERSON]

	MediaUtils: Better sample fetching when paused

Change 3592761 on 2017/08/17 by [REDACTED_PERSON]

	MediaUtils: Fixed overflow detection in sample sinks

Change 3592762 on 2017/08/17 by [REDACTED_PERSON]

	ImgMedia: Sending end reached event when looping

Change 3592885 on 2017/08/17 by [REDACTED_PERSON]

	ImgMedia: Fixed reverse play

Change 3592887 on 2017/08/17 by [REDACTED_PERSON]

	MediaUtils: Better sample range calculation for audio samples

Change 3593010 on 2017/08/17 by [REDACTED_PERSON]

	ImgMedia: Fixed async loading of non-EXR image sequences

Change 3593193 on 2017/08/17 by [REDACTED_PERSON]

	AndroidMedia: Fixed typo

Change 3593230 on 2017/08/17 by [REDACTED_PERSON]

	Media: Implemented flushing for player sample queues

Change 3593346 on 2017/08/17 by [REDACTED_PERSON]

	Media: Proper sample processing for reverse playback

Change 3593482 on 2017/08/17 by [REDACTED_PERSON]

	Switch: Fixed sample range check for reverse playback

Change 3594428 on 2017/08/17 by [REDACTED_PERSON]

	PS4Media: Fixed video not playing/crashing

Change 3595404 on 2017/08/17 by [REDACTED_PERSON]

	MediaAssets: Not requesting audio samples if not playing

Change 3595624 on 2017/08/17 by [REDACTED_PERSON]

	PS4Media: Fixed no looping after track switching

Change 3595704 on 2017/08/17 by [REDACTED_PERSON]

	MediaUtils: Improved sample queue flushing

	Eventually this needs to be event-driven.

Change 3595706 on 2017/08/17 by [REDACTED_PERSON]

	PS4Media: Flushing queues when track switching

Change 3595909 on 2017/08/17 by [REDACTED_PERSON]

	Matinee: Fix Matinee redirect missing as UClass has been GC'ed (GitHub PR #3382)

	https://github.com/EpicGames/UnrealEngine/pull/3382

	#rn Fix Matinee redirect missing as UClass has been GC'ed (GitHub PR #3382)

Change 3597480 on 2017/08/18 by [REDACTED_PERSON]

	ImgMediaPlayer: Fixed presentation time calculation for very large delta times

Change 3597669 on 2017/08/18 by [REDACTED_PERSON]

	ImgMedia: Setting player stopped instead of paused when end reached

Change 3597709 on 2017/08/18 by [REDACTED_PERSON]

	SwitchMedia: Fixed audio sample duration calculation

Change 3598479 on 2017/08/18 by [REDACTED_PERSON]

	ImgMedia: Ensuring that first/last frame is played when starting playback forward/reverse; improved state checking; user logging

Change 3598629 on 2017/08/18 by [REDACTED_PERSON]

	MediaUtils: Properly handling pending flushes when peeking sample queues

Change 3598633 on 2017/08/18 by [REDACTED_PERSON]

	WmfMedia: Fixed occasional WMF deadlock when scrubbing

Change 3598653 on 2017/08/18 by [REDACTED_PERSON]

	MediaUtils: Fixed audio resampling in reverse playback

Change 3598659 on 2017/08/18 by [REDACTED_PERSON]

	MediaPlayerEditor: Creating sound output only if audio device available

Change 3598688 on 2017/08/18 by [REDACTED_PERSON]

	MediaUtils: Flushing sample queues on direction change regardless of paused state

Change 3599444 on 2017/08/20 by [REDACTED_PERSON]

	WmfMedia: Added missing check for rate control when quering supported rates

Change 3603661 on 2017/08/22 by [REDACTED_PERSON]

	WmfMedia: Fixed NV12 output

Change 3604345 on 2017/08/23 by [REDACTED_PERSON]

	MediaUtils: Added subtitle samples to media sample collection; fixed documentation

Change 3604987 on 2017/08/23 by [REDACTED_PERSON]

	PS4Media: Setting correct track format types

Change 3605117 on 2017/08/23 by [REDACTED_PERSON]

	Fix issue where presentation time was being treated as milliseconds instead of microseconds.

Change 3605128 on 2017/08/23 by [REDACTED_PERSON]

	Fix potential issue calculating total (looped) play time after seek or reset.
	Track last video presentation time

Change 3605139 on 2017/08/23 by [REDACTED_PERSON]

	Support audio type DType_Setup.
	Fix issue with procedural sounds popping. Consume more samples if necessary.

Change 3605197 on 2017/08/23 by [REDACTED_PERSON]

	Media: Fixed track display name formatting (UE-48767)

Change 3605817 on 2017/08/23 by [REDACTED_PERSON]

	Auto set sampler type on drop on actor.

Change 3605999 on 2017/08/23 by [REDACTED_PERSON]

	MfMedia: Restarting source reader when switching tracks

Change 3606416 on 2017/08/23 by [REDACTED_PERSON]

	PS4Media: Various player improvements

Change 3607656 on 2017/08/24 by [REDACTED_PERSON]

	WmfMedia: Improved logging for Seek and SetRate

Change 3607855 on 2017/08/24 by [REDACTED_PERSON]

	MfMedia: Better verbose logging; validating seeks; reduced nesting; fixed track switching\

Change 3608029 on 2017/08/24 by [REDACTED_PERSON]

	PS4Media: Fixed streaming media not playing; improved result checks & logging

Change 3608473 on 2017/08/24 by [REDACTED_PERSON]

	MfMedia: Fixed audio sample duration calculation

Change 3609316 on 2017/08/24 by [REDACTED_PERSON]

	Provide external texture RotationScale and Offset from sample

Change 3610067 on 2017/08/25 by [REDACTED_PERSON]

	Mac: Editor locks up when Track switching and Mac: Media Audio continues to play after closing player.

	- Fix for creating extra OS AV media playback objects so audio playback not stopping due to over retained objects.
	- Trying to access OS AV Media Playback object (esp. the currentTime value) across mutliple threads at once seems to cause deadlock in the Apple libraries - moved location of currentTime acess point to Media Ticker
- this sets the current time and is then just returned in the getter when on game main.
	- Trying to playback audio via OS media playback and through the engine at the same time when enabling the selected audio track.  Engine version via sample buffers via AVAssetReaderTrackOutput doesn't seem to be
working as intended with the new audio frame work audio sinks - so reverting to OS AV media playback for Mac at the moment until better fix for this can be put in.

Change 3610267 on 2017/08/25 by [REDACTED_PERSON]

	PS4Media: Added missing track selection validation

Change 3610399 on 2017/08/25 by [REDACTED_PERSON]

	PS4Media: Fixed infinite player re-initialization loop if track selection failed

Change 3610809 on 2017/08/25 by [REDACTED_PERSON]

	Fixes for Android media player and camera player
	- buffer samples copy proper bytecount
	- realloc sets new buffer ptr
	- refcount of Java buffers fixed

Change 3610953 on 2017/08/25 by [REDACTED_PERSON]

	Fix audio disable before play on Android media player

Change 3611405 on 2017/08/25 by [REDACTED_PERSON]

	WmfMedia: Resetting supported rates if RateSupport unavailable

Change 3611406 on 2017/08/25 by [REDACTED_PERSON]

	MfMedia: Resetting supported rates if RateSupport unavailable

Change 3611453 on 2017/08/25 by [REDACTED_PERSON]

	Android external texture extension updates for compatibility

Change 3611719 on 2017/08/26 by [REDACTED_PERSON]

	Media: Added media event for completed seek operations; flushing sinks on seek

Change 3611764 on 2017/08/26 by [REDACTED_PERSON]

	AvfMedia: Setting player to preparing state prior to initializing tracks asynchronously

Change 3611802 on 2017/08/26 by [REDACTED_PERSON]

	Media: Exposed error state in media player

Change 3611803 on 2017/08/26 by [REDACTED_PERSON]

	MediaPlayerEditor: Showing error state in UI

Change 3611887 on 2017/08/26 by [REDACTED_PERSON]

	MediaUtils: Corrected media sample sink overflow check

Change 3611892 on 2017/08/27 by [REDACTED_PERSON]

	WmfMedia: Leaving session in error state after error

Change 3611929 on 2017/08/27 by [REDACTED_PERSON]

	MediaPlayerEditor: Fixed sound stopping on looping

Change 3611930 on 2017/08/27 by [REDACTED_PERSON]

	MfMedia: Rewrote async sample processing to fix various playback issues

Change 3611942 on 2017/08/27 by [REDACTED_PERSON]

	Media: Sending suspension event when playback ended

Change 3611957 on 2017/08/27 by [REDACTED_PERSON]

	UnrealEd: Allowed transient assets to be resaved to disk

Change 3611981 on 2017/08/27 by [REDACTED_PERSON]

	PS4Media: Rewrote sample processing to fix various playback issues

Change 3612035 on 2017/08/27 by [REDACTED_PERSON]

	UnrealEd: Picking standard default names when resaving transient packages

Change 3612045 on 2017/08/27 by [REDACTED_PERSON]

	Media: Fixed playlists not getting saved correctly from Editor (UE-35382)

Change 3612212 on 2017/08/28 by [REDACTED_PERSON]

	Fix for Multitrack Video samples continue to Play on Open after Play on Open is toggled off.  Current Rate was not getting reset back to 0.0 on media close.

Change 3613531 on 2017/08/28 by [REDACTED_PERSON]

	MediaAssets: Fixed external texture related crash on shutdown (UE-48918)

	Also no longer creating clock sink for media player CDO

Change 3613677 on 2017/08/28 by [REDACTED_PERSON]

	Fixed crash at exit in Mac editor caused by double release of AVPlayerItem object

Change 3615917 on 2017/08/29 by [REDACTED_PERSON]

	MfMedia: Added compile options for DXVA and falling back to synchronous sample reading

diff --git a/Engine/Build/Android/Java/proguard-project.txt b/Engine/Build/Android/Java/proguard-project.txt
--- a/Engine/Build/Android/Java/proguard-project.txt
+++ b/Engine/Build/Android/Java/proguard-project.txt
@@ -53,6 +53,9 @@
 -keep class com.epicgames.ue4.GameActivity$JavaAchievement {
 	public *;
 }
+-keep class com.epicgames.ue4.MediaPlayer14$FrameUpdateInfo {
+	public *;
+}
 -keep class com.epicgames.ue4.MediaPlayer14$AudioTrackInfo {
 	public *;
 }
diff --git a/Engine/Build/Android/Java/src/com/epicgames/ue4/MediaPlayer14.java b/Engine/Build/Android/Java/src/com/epicgames/ue4/MediaPlayer14.java
--- a/Engine/Build/Android/Java/src/com/epicgames/ue4/MediaPlayer14.java
+++ b/Engine/Build/Android/Java/src/com/epicgames/ue4/MediaPlayer14.java
@@ -25,8 +25,26 @@
 public class MediaPlayer14
 	extends android.media.MediaPlayer
 {
-	private boolean VulkanRenderer = false;									
+	private boolean SwizzlePixels = true;
+	private boolean VulkanRenderer = false;
+	private boolean Looping = false;
+	private boolean AudioEnabled = true;
 	private volatile boolean WaitOnBitmapRender = false;
+	private volatile boolean Prepared = false;
+	private volatile boolean Completed = false;
 
+	private BitmapRenderer mBitmapRenderer = null;
+	private OESTextureRenderer mOESTextureRenderer = null;
+
+	public class FrameUpdateInfo {
+		int CurrentPosition;
+		boolean FrameReady;
+		boolean RegionChanged;
+		float UScale;
+		float UOffset;
+		float VScale;
+		float VOffset;
+	}
+
 	public class AudioTrackInfo {
 		public int Index;
@@ -58,24 +76,73 @@
 	private ArrayList<AudioTrackInfo> audioTracks = new ArrayList<AudioTrackInfo>();
 	private ArrayList<VideoTrackInfo> videoTracks = new ArrayList<VideoTrackInfo>();
 
-	public MediaPlayer14(boolean vulkanRenderer)
+	// ======================================================================================
+
+	public MediaPlayer14(boolean swizzlePixels, boolean vulkanRenderer)
 	{
+		SwizzlePixels = swizzlePixels;
 		VulkanRenderer = vulkanRenderer;
-		
-		GameActivity.Log.debug("MediaPlayer14: initialized,  VulkanRenderer: " + VulkanRenderer);
-
 		WaitOnBitmapRender = false;
+		AudioEnabled = true;
 
-		 setOnErrorListener(new OnErrorListener()
-		 {
+		setOnErrorListener(new MediaPlayer.OnErrorListener() {
 			@Override
-			public boolean onError(    MediaPlayer mp,    int what,    int extra)
-			{
-			  GameActivity.Log.warn("Error occurred while playing audio. ("+what+", "+extra+")");
-			  return true;
+			public boolean onError(MediaPlayer mp, int what, int extra) {
+				GameActivity.Log.debug("MediaPlayer14: onError returned what=" + what + ", extra=" + extra);
+				return true;
 			}
-		}
-		);
+		});
+
+		setOnPreparedListener(new MediaPlayer.OnPreparedListener() {
+			@Override
+			public void onPrepared(MediaPlayer player) {
+				synchronized(player)
+				{
+					Prepared = true;
+				}
+			}
+		});
+
+		setOnCompletionListener(new MediaPlayer.OnCompletionListener() {
+			@Override
+			public void onCompletion(MediaPlayer player) {
+				synchronized(player)
+				{
+					if (Looping)
+					{
+						seekTo(0);
+						start();
+					}
+					Completed = true;
+				}
+			}
+		});
 	}
 
+	public boolean isPrepared()
+	{
+		boolean result;
+		synchronized(this)
+		{
+			result = Prepared;
+		}
+		return result;
+	}
+
+	public boolean didComplete()
+	{
+		boolean result;
+		synchronized(this)
+		{
+			result = Completed;
+			Completed = false;
+		}
+		return result;
+	}
+
+	public boolean isLooping()
+	{
+		return Looping;
+	}
+
 	private void updateTrackInfo(MediaExtractor extractor)
@@ -158,12 +225,20 @@
 			java.lang.InterruptedException,
 			java.util.concurrent.ExecutionException
 	{
+		synchronized(this)
+		{
+			Prepared = false;
+			Completed = false;
+		}
+		Looping = false;
+		AudioEnabled = true;
 		audioTracks.clear();
 		videoTracks.clear();
 
 		try
 		{
 			setDataSource(UrlPath);
+			releaseOESTextureRenderer();
 			releaseBitmapRenderer();
 			if (android.os.Build.VERSION.SDK_INT >= 16)
 			{
@@ -191,6 +266,13 @@
 			java.lang.InterruptedException,
 			java.util.concurrent.ExecutionException
 	{
+		synchronized(this)
+		{
+			Prepared = false;
+			Completed = false;
+		}
+		Looping = false;
+		AudioEnabled = true;
 		audioTracks.clear();
 		videoTracks.clear();
 
@@ -203,6 +285,7 @@
 			}
 			RandomAccessFile data = new RandomAccessFile(f, "r");
 			setDataSource(data.getFD(), offset, size);
+			releaseOESTextureRenderer();
 			releaseBitmapRenderer();
 
 			if (android.os.Build.VERSION.SDK_INT >= 16)
@@ -230,6 +313,13 @@
 		throws java.lang.InterruptedException,
 			java.util.concurrent.ExecutionException
 	{
+		synchronized(this)
+		{
+			Prepared = false;
+			Completed = false;
+		}
+		Looping = false;
+		AudioEnabled = true;
 		audioTracks.clear();
 		videoTracks.clear();
 
@@ -237,6 +327,7 @@
 		{
 			AssetFileDescriptor assetFD = assetManager.openFd(assetPath);
 			setDataSource(assetFD.getFileDescriptor(), offset, size);
+			releaseOESTextureRenderer();
 			releaseBitmapRenderer();
 
 			if (android.os.Build.VERSION.SDK_INT >= 16)
@@ -268,6 +359,11 @@
 		mVideoEnabled = enabled;
 		if (mVideoEnabled)
 		{
+			if (null != mOESTextureRenderer && null != mOESTextureRenderer.getSurface())
+			{
+				setSurface(mOESTextureRenderer.getSurface());
+			}
+
 			if (null != mBitmapRenderer && null != mBitmapRenderer.getSurface())
 			{
 				setSurface(mBitmapRenderer.getSurface());
@@ -283,6 +379,7 @@
 	
 	public void setAudioEnabled(boolean enabled)
 	{
+		AudioEnabled = enabled;
 		if (enabled)
 		{
 			setVolume(1,1);
@@ -295,30 +392,154 @@
 	
 	public boolean didResolutionChange()
 	{
-		if (null == mBitmapRenderer)
+		if (null != mOESTextureRenderer)
 		{
+			return mOESTextureRenderer.resolutionChanged();
+		}
+		if (null != mBitmapRenderer)
+		{
+			return mBitmapRenderer.resolutionChanged();
+		}
+		return false;
+	}
+
+	public int getExternalTextureId()
+	{
+		if (null != mOESTextureRenderer)
+		{
+			return mOESTextureRenderer.getExternalTextureId();
+		}
+		if (null != mBitmapRenderer)
+		{
+			return mBitmapRenderer.getExternalTextureId();
+		}
+		return -1;
+	}
+
+	public void start()
+	{
+		synchronized(this)
+		{
+			Completed = false;
+			super.start();
+		}
+	}
+
+	public void stop()
+	{
+		synchronized(this)
+		{
+			Completed = false;
+			super.stop();
+		}
+	}
+
+	public void seekTo(int position)
+	{
+		synchronized (this)
+		{
+			Completed = false;
+			super.seekTo(position);
+		}
+	}
+
+	public void setLooping(boolean looping)
+	{
+		// don't set on player
+		Looping = looping;
+	}
+
+	public void release()
+	{
+		if (null != mOESTextureRenderer)
+		{
+			while (WaitOnBitmapRender) ;
+			releaseOESTextureRenderer();
+		}
+		if (null != mBitmapRenderer)
+		{
+			while (WaitOnBitmapRender) ;
+			releaseOESTextureRenderer();
+		}
+		super.release();
+	}
+
+	public void reset()
+	{
+		if (null != mOESTextureRenderer)
+		{
+			while (WaitOnBitmapRender) ;
+			releaseOESTextureRenderer();
+		}
+		if (null != mBitmapRenderer)
+		{
+			while (WaitOnBitmapRender) ;
+			releaseBitmapRenderer();
+		}
+		super.reset();
+	}
+
+	// ======================================================================================
+
+	private boolean CreateBitmapRenderer()
+	{
+		releaseBitmapRenderer();
+
+		mBitmapRenderer = new BitmapRenderer(SwizzlePixels, VulkanRenderer);
+		if (!mBitmapRenderer.isValid())
+		{
+			mBitmapRenderer = null;
 			return false;
 		}
-		return mBitmapRenderer.resolutionChanged();
+		
+		// set this here as the size may have been set before the GL resources were created.
+		mBitmapRenderer.setSize(getVideoWidth(),getVideoHeight());
+
+		setOnVideoSizeChangedListener(new android.media.MediaPlayer.OnVideoSizeChangedListener() {
+			public void onVideoSizeChanged(android.media.MediaPlayer player, int w, int h)
+			{
+//				GameActivity.Log.debug("VIDEO SIZE CHANGED: " + w + " x " + h);
+				if (null != mBitmapRenderer)
+				{
+					mBitmapRenderer.setSize(w,h);
+				}
+			}
+		});
+		setVideoEnabled(true);
+		if (AudioEnabled)
+		{
+			setAudioEnabled(true);
+		}
+		return true;
 	}
 
+	void releaseBitmapRenderer()
+	{
+		if (null != mBitmapRenderer)
+		{
+			mBitmapRenderer.release();
+			mBitmapRenderer = null;
+			setSurface(null);
+			setOnVideoSizeChangedListener(null);
+		}
+	}
+
 	public void initBitmapRenderer()
 	{
 		// if not already allocated.
 		// Create bitmap renderer's gl resources in the renderer thread.
-		  if (null == mBitmapRenderer)
-		  {
+		if (null == mBitmapRenderer)
+		{
 			if (!CreateBitmapRenderer())
 			{
 				GameActivity.Log.warn("initBitmapRenderer failed to alloc mBitmapRenderer ");
 				reset();
 			  }
-		  }
+		}
 	}
 
 	public java.nio.Buffer getVideoLastFrameData()
 	{
-
 		initBitmapRenderer();
 		if (null != mBitmapRenderer)
 		{
@@ -335,6 +556,7 @@
 
 	public boolean getVideoLastFrame(int destTexture)
 	{
+//		GameActivity.Log.debug("getVideoLastFrame: " + destTexture);
 		initBitmapRenderer();
 		if (null != mBitmapRenderer)
 		{
@@ -349,26 +571,6 @@
 		}
 	}
 
-	public void release()
-	{
-		if (null != mBitmapRenderer)
-		{
-			while (WaitOnBitmapRender) ;
-			releaseBitmapRenderer();
-		}
-		super.release();
-	}
-
-	public void reset()
-	{
-		if (null != mBitmapRenderer)
-		{
-			while (WaitOnBitmapRender) ;
-			releaseBitmapRenderer();
-		}
-		super.reset();
-	}
-	
 	/*
 		All this internal surface view does is manage the
 		offscreen bitmap that the media player decoding can
@@ -408,14 +610,14 @@
 
 		private boolean mCreatedEGLDisplay = false;
 
-		public BitmapRenderer(boolean isVulkan)
+		public BitmapRenderer(boolean swizzlePixels, boolean isVulkan)
 		{
+			mSwizzlePixels = swizzlePixels;
 			mVulkanRenderer = isVulkan;
-			mUseOwnContext = true;
 
 			mEglSurface = EGL14.EGL_NO_SURFACE;
 			mEglContext = EGL14.EGL_NO_CONTEXT;
-			mEglDisplay = EGL14.EGL_NO_DISPLAY;
+			mUseOwnContext = true;
 
 			if (mVulkanRenderer)
 			{
@@ -693,6 +895,11 @@
 			return mSurface;
 		}
 
+		public int getExternalTextureId()
+		{
+			return mTextureID;
+		}
+
 		// NOTE: Synchronized with updateFrameData to prevent frame
 		// updates while the surface may need to get reallocated.
 		public void setSize(int width, int height)
@@ -1059,23 +1266,38 @@
 			return true;
 		}
 
+		private void showGlError(String op, int error)
+		{
+			switch (error)
+			{
+				case GLES20.GL_INVALID_ENUM:						GameActivity.Log.error("MediaPlayer$BitmapRenderer: " + op + ": glGetError GL_INVALID_ENUM");  break;
+				case GLES20.GL_INVALID_OPERATION:					GameActivity.Log.error("MediaPlayer$BitmapRenderer: " + op + ": glGetError GL_INVALID_OPERATION");  break;
+				case GLES20.GL_INVALID_FRAMEBUFFER_OPERATION:		GameActivity.Log.error("MediaPlayer$BitmapRenderer: " + op + ": glGetError GL_INVALID_FRAMEBUFFER_OPERATION");  break;
+				case GLES20.GL_INVALID_VALUE:						GameActivity.Log.error("MediaPlayer$BitmapRenderer: " + op + ": glGetError GL_INVALID_VALUE");  break;
+				case GLES20.GL_FRAMEBUFFER_INCOMPLETE_ATTACHMENT:	GameActivity.Log.error("MediaPlayer$BitmapRenderer: " + op + ": glGetError GL_FRAMEBUFFER_INCOMPLETE_ATTACHMENT");  break;
+				case GLES20.GL_FRAMEBUFFER_INCOMPLETE_DIMENSIONS:	GameActivity.Log.error("MediaPlayer$BitmapRenderer: " + op + ": glGetError GL_FRAMEBUFFER_INCOMPLETE_DIMENSIONS");  break;
+				case GLES20.GL_FRAMEBUFFER_UNSUPPORTED:				GameActivity.Log.error("MediaPlayer$BitmapRenderer: " + op + ": glGetError GL_FRAMEBUFFER_UNSUPPORTED");  break;
+				case GLES20.GL_OUT_OF_MEMORY:						GameActivity.Log.error("MediaPlayer$BitmapRenderer: " + op + ": glGetError GL_OUT_OF_MEMORY");  break;
+				default:											GameActivity.Log.error("MediaPlayer$BitmapRenderer: " + op + ": glGetError " + error);
+			}
+		}
+
 		private void glVerify(String op)
 		{
 			int error;
 			while ((error = GLES20.glGetError()) != GLES20.GL_NO_ERROR)
 			{
-				GameActivity.Log.error("MediaPlayer$BitmapRenderer: " + op + ": glGetError " + error);
+				showGlError(op, error);
 				throw new RuntimeException(op + ": glGetError " + error);
 			}
 		}
 
-
 		private void glWarn(String op)
 		{
 			int error;
 			while ((error = GLES20.glGetError()) != GLES20.GL_NO_ERROR)
 			{
-				GameActivity.Log.warn("MediaPlayer$BitmapRenderer: " + op + ": glGetError " + error);
+				showGlError(op, error);
 			}
 		}
 
@@ -1144,44 +1366,236 @@
 		}
 	};
 
-	private boolean CreateBitmapRenderer()
+	public native void nativeClearCachedAttributeState(int PositionAttrib, int TexCoordsAttrib);
+
+	// ======================================================================================
+
+	private boolean CreateOESTextureRenderer(int OESTextureId)
 	{
-		releaseBitmapRenderer();
+		releaseOESTextureRenderer();
 
-		mBitmapRenderer = new BitmapRenderer(VulkanRenderer);
-		if (!mBitmapRenderer.isValid())
+		mOESTextureRenderer = new OESTextureRenderer(OESTextureId);
+		if (!mOESTextureRenderer.isValid())
 		{
-			mBitmapRenderer = null;
+			mOESTextureRenderer = null;
 			return false;
 		}
 		
 		// set this here as the size may have been set before the GL resources were created.
-		mBitmapRenderer.setSize(getVideoWidth(),getVideoHeight());
+		mOESTextureRenderer.setSize(getVideoWidth(),getVideoHeight());
 
 		setOnVideoSizeChangedListener(new android.media.MediaPlayer.OnVideoSizeChangedListener() {
 			public void onVideoSizeChanged(android.media.MediaPlayer player, int w, int h)
 			{
 //				GameActivity.Log.debug("VIDEO SIZE CHANGED: " + w + " x " + h);
-				mBitmapRenderer.setSize(w,h);
+				if (null != mOESTextureRenderer)
+				{
+					mOESTextureRenderer.setSize(w,h);
+				}
 			}
-			});
+		});
 		setVideoEnabled(true);
-		setAudioEnabled(true);
+		if (AudioEnabled)
+		{
+			setAudioEnabled(true);
+		}
 		return true;
 	}
 
-	void releaseBitmapRenderer()
+	void releaseOESTextureRenderer()
 	{
-		if (null != mBitmapRenderer)
+		if (null != mOESTextureRenderer)
 		{
-			mBitmapRenderer.release();
-			mBitmapRenderer = null;
+			mOESTextureRenderer.release();
+			mOESTextureRenderer = null;
 			setSurface(null);
 			setOnVideoSizeChangedListener(null);
 		}
 	}
+	
+	public FrameUpdateInfo updateVideoFrame(int externalTextureId)
+	{
+		if (null == mOESTextureRenderer)
+		{
+			if (!CreateOESTextureRenderer(externalTextureId))
+			{
+				GameActivity.Log.warn("updateVideoFrame failed to alloc mOESTextureRenderer ");
+				reset();
+				return null;
+			}
+		}
 
-	private BitmapRenderer mBitmapRenderer = null;
+		WaitOnBitmapRender = true;
+		FrameUpdateInfo result = mOESTextureRenderer.updateVideoFrame();
+		WaitOnBitmapRender = false;
+		return result;
+	}
 
+	/*
+		This handles events for our OES texture
+	*/
+	class OESTextureRenderer
+		implements android.graphics.SurfaceTexture.OnFrameAvailableListener
+	{
+		private android.graphics.SurfaceTexture mSurfaceTexture = null;
+		private int mTextureWidth = -1;
+		private int mTextureHeight = -1;
+		private android.view.Surface mSurface = null;
+		private boolean mFrameAvailable = false;
+		private int mTextureID = -1;
+		private float[] mTransformMatrix = new float[16];
+		private boolean mTextureSizeChanged = true;
+
+		private float mUScale = 1.0f;
+		private float mVScale = -1.0f;
+		private float mUOffset = 0.0f;
+		private float mVOffset = 1.0f;
+
+		public OESTextureRenderer(int OESTextureId)
+		{
+			mTextureID = OESTextureId;
+
+			mSurfaceTexture = new android.graphics.SurfaceTexture(mTextureID);
+			mSurfaceTexture.setOnFrameAvailableListener(this);
+			mSurface = new android.view.Surface(mSurfaceTexture);
+		}
+
+		public void release()
+		{
+			if (null != mSurface)
+			{
+				mSurface.release();
+				mSurface = null;
+			}
+			if (null != mSurfaceTexture)
+			{
+				mSurfaceTexture.release();
+				mSurfaceTexture = null;
+			}
+		}
+
+		public boolean isValid()
+		{
+			return mSurfaceTexture != null;
+		}
+
+		public void onFrameAvailable(android.graphics.SurfaceTexture st)
+		{
+			synchronized(this)
+			{
+				mFrameAvailable = true;
+			}
+		}
+
+		public android.graphics.SurfaceTexture getSurfaceTexture()
+		{
+			return mSurfaceTexture;
+		}
+
+		public android.view.Surface getSurface()
+		{
+			return mSurface;
+		}
+
+		public int getExternalTextureId()
+		{
+			return mTextureID;
+		}
+
+		// NOTE: Synchronized with updateFrameData to prevent frame
+		// updates while the surface may need to get reallocated.
+		public void setSize(int width, int height)
+		{
+			synchronized(this)
+			{
+				if (width != mTextureWidth ||
+					height != mTextureHeight)
+				{
+					mTextureWidth = width;
+					mTextureHeight = height;
+					mTextureSizeChanged = true;
+				}
+			}
+		}
+
+		public boolean resolutionChanged()
+		{
+			boolean changed;
+			synchronized(this)
+			{
+				changed = mTextureSizeChanged;
+				mTextureSizeChanged = false;
+			}
+			return changed;
+		}
+
+		public FrameUpdateInfo updateVideoFrame()
+		{
+			synchronized(this)
+			{
+				return getFrameUpdateInfo();
+			}
+		}
+
+		private FrameUpdateInfo getFrameUpdateInfo()
+		{
+			FrameUpdateInfo frameUpdateInfo = new FrameUpdateInfo();
+
+			frameUpdateInfo.CurrentPosition = getCurrentPosition();
+			frameUpdateInfo.FrameReady = false;
+			frameUpdateInfo.RegionChanged = false;
+			frameUpdateInfo.UScale = mUScale;
+			frameUpdateInfo.UOffset = mUOffset;
+
+			// note: the matrix has V flipped
+			frameUpdateInfo.VScale = -mVScale;
+			frameUpdateInfo.VOffset = 1.0f - mVOffset;
+
+			if (!mFrameAvailable)
+			{
+				// We only return fresh data when we generate it. At other
+				// time we return nothing to indicate that there was nothing
+				// new to return. The media player deals with this by keeping
+				// the last frame around and using that for rendering.
+				return frameUpdateInfo;
+			}
+			mFrameAvailable = false;
+			if (null == mSurfaceTexture)
+			{
+				// Can't update if there's no surface to update into.
+				return frameUpdateInfo;
+			}
+
+			frameUpdateInfo.FrameReady = true;
+
+			// Get the latest video texture frame.
+			mSurfaceTexture.updateTexImage();
+			mSurfaceTexture.getTransformMatrix(mTransformMatrix);
+
+			if (mUScale != mTransformMatrix[0] ||
+				mVScale != mTransformMatrix[5] ||
+				mUOffset != mTransformMatrix[12] ||
+				mVOffset != mTransformMatrix[13])
+			{
+				mUScale = mTransformMatrix[0];
+				mVScale = mTransformMatrix[5];
+				mUOffset = mTransformMatrix[12];
+				mVOffset = mTransformMatrix[13];
+
+				frameUpdateInfo.RegionChanged = true;
+				frameUpdateInfo.UScale = mUScale;
+				frameUpdateInfo.UOffset = mUOffset;
+
+				// note: the matrix has V flipped
+				frameUpdateInfo.VScale = -mVScale;
+				frameUpdateInfo.VOffset = 1.0f - mVOffset;
+			}
+
+			return frameUpdateInfo;
+		}
+	};
+
+	// ======================================================================================
+
 	public AudioTrackInfo[] GetAudioTracks()
 	{
@@ -1248,6 +1662,7 @@
 
 		AudioTrackInfo[] AudioTracks = new AudioTrackInfo[1];
 
+		AudioTracks[0] = new AudioTrackInfo();
 		AudioTracks[0].Index = 0;
 		AudioTracks[0].MimeType = "audio/unknown";
 		AudioTracks[0].DisplayName = "Audio Track 0 (Stream 0)";
@@ -1381,7 +1796,7 @@
 		{
 			VideoTrackInfo[] VideoTracks = new VideoTrackInfo[1];
 
-			VideoTrackInfo VideoTrack = new VideoTrackInfo();
+			VideoTracks[0] = new VideoTrackInfo();
 			VideoTracks[0].Index = 0;
 			VideoTracks[0].MimeType = "video/unknown";
 			VideoTracks[0].DisplayName = "Video Track 0 (Stream 0)";
@@ -1398,7 +1813,4 @@
 
 		return VideoTracks;
 	}
-
-	public native void nativeClearCachedAttributeState(int PositionAttrib, int TexCoordsAttrib);
-
 }
diff --git a/Engine/Plugins/Developer/GitSourceControl/Source/GitSourceControl/Private/GitSourceControlUtils.cpp b/Engine/Plugins/Developer/GitSourceControl/Source/GitSourceControl/Private/GitSourceControlUtils.cpp
--- a/Engine/Plugins/Developer/GitSourceControl/Source/GitSourceControl/Private/GitSourceControlUtils.cpp
+++ b/Engine/Plugins/Developer/GitSourceControl/Source/GitSourceControl/Private/GitSourceControlUtils.cpp
@@ -108,10 +108,26 @@ static bool RunCommandInternalRaw(const FString& InCommand, const FString& InPat
 	FPlatformProcess::ExecProcess(*InPathToGitBinary, *FullCommand, &ReturnCode, &OutResults, &OutErrors);
 	
 #if UE_BUILD_DEBUG
-	UE_LOG(LogSourceControl, Log, TEXT("RunCommandInternalRaw(%s): OutResults=\n%s"), *InCommand, *OutResults)
+
+	if (OutResults.IsEmpty())
+	{
+		UE_LOG(LogSourceControl, Log, TEXT("RunCommandInternalRaw: 'OutResults=n/a'"));
+	}
+	else
+	{
+		UE_LOG(LogSourceControl, Log, TEXT("RunCommandInternalRaw(%s): OutResults=\n%s"), *InCommand, *OutResults);
+	}
+
 	if(ReturnCode != 0)
 	{
-		UE_LOG(LogSourceControl, Warning, TEXT("RunCommandInternalRaw(%s): OutErrors=\n%s"), *InCommand, *OutErrors);
+		if (OutErrors.IsEmpty())
+		{
+			UE_LOG(LogSourceControl, Warning, TEXT("RunCommandInternalRaw: 'OutErrors=n/a'"));
+		}
+		else
+		{
+			UE_LOG(LogSourceControl, Warning, TEXT("RunCommandInternalRaw(%s): OutErrors=\n%s"), *InCommand, *OutErrors);
+		}
 	}
 #endif
 
diff --git a/Engine/Plugins/Developer/PerforceSourceControl/Source/PerforceSourceControl/Private/PerforceSourceControlLabel.cpp b/Engine/Plugins/Developer/PerforceSourceControl/Source/PerforceSourceControl/Private/PerforceSourceControlLabel.cpp
--- a/Engine/Plugins/Developer/PerforceSourceControl/Source/PerforceSourceControl/Private/PerforceSourceControlLabel.cpp
+++ b/Engine/Plugins/Developer/PerforceSourceControl/Source/PerforceSourceControl/Private/PerforceSourceControlLabel.cpp
@@ -33,7 +33,7 @@ static void ParseFilesResults(const FP4RecordSet& InRecords, TArray< TSharedRef<
 		Revision->RevisionNumber = FCString::Atoi(*RevisionNumber);
 		Revision->ChangelistNumber = FCString::Atoi(*ChangelistNumber);
 		Revision->Action = Action;
-		Revision->Date = FDateTime(1970, 1, 1, 0, 0, 0, 0) + FTimespan(0, 0, FCString::Atoi(*Date));
+		Revision->Date = FDateTime(1970, 1, 1, 0, 0, 0, 0) + FTimespan::FromSeconds(FCString::Atoi(*Date));
 
 		OutRevisions.Add(Revision);
 	}
diff --git a/Engine/Plugins/Developer/PerforceSourceControl/Source/PerforceSourceControl/Private/PerforceSourceControlOperations.cpp b/Engine/Plugins/Developer/PerforceSourceControl/Source/PerforceSourceControl/Private/PerforceSourceControlOperations.cpp
--- a/Engine/Plugins/Developer/PerforceSourceControl/Source/PerforceSourceControl/Private/PerforceSourceControlOperations.cpp
+++ b/Engine/Plugins/Developer/PerforceSourceControl/Source/PerforceSourceControl/Private/PerforceSourceControlOperations.cpp
@@ -863,7 +863,7 @@ static void ParseHistoryResults(const FP4RecordSet& InRecords, const TArray<FPer
 				Revision->ClientSpec = ClientSpec;
 				Revision->Action = Action;
 				Revision->BranchSource = BranchSource;
-				Revision->Date = FDateTime(1970, 1, 1, 0, 0, 0, 0) + FTimespan(0, 0, FCString::Atoi(*Date));
+				Revision->Date = FDateTime(1970, 1, 1, 0, 0, 0, 0) + FTimespan::FromSeconds(FCString::Atoi(*Date));
 				Revision->FileSize = FCString::Atoi(*FileSize);
 
 				Revisions.Add(Revision);
diff --git a/Engine/Plugins/Editor/EpicSurvey/Source/EpicSurvey/Private/EpicSurvey.cpp b/Engine/Plugins/Editor/EpicSurvey/Source/EpicSurvey/Private/EpicSurvey.cpp
--- a/Engine/Plugins/Editor/EpicSurvey/Source/EpicSurvey/Private/EpicSurvey.cpp
+++ b/Engine/Plugins/Editor/EpicSurvey/Source/EpicSurvey/Private/EpicSurvey.cpp
@@ -18,7 +18,8 @@
 #include "Framework/MultiBox/MultiBoxBuilder.h"
 #include "LevelEditor.h"
 #include "EngineAnalytics.h"
-#include "Interfaces/IImageWrapperModule.h"
+#include "IImageWrapper.h"
+#include "IImageWrapperModule.h"
 #include "Interfaces/IMainFrameModule.h"
 #include "SurveyTitleCdnStorage.h"
 #include "AnalyticsEventAttribute.h"
@@ -91,7 +92,7 @@ TSharedPtr< FSlateDynamicImageBrush > FEpicSurvey::LoadRawDataAsBrush( FName Res
 	bool bSucceeded = false;
 	TArray<uint8> DecodedImage;
 	IImageWrapperModule& ImageWrapperModule = FModuleManager::LoadModuleChecked<IImageWrapperModule>( FName("ImageWrapper") );
-	IImageWrapperPtr ImageWrapper = ImageWrapperModule.CreateImageWrapper( EImageFormat::PNG );
+	TSharedPtr<IImageWrapper> ImageWrapper = ImageWrapperModule.CreateImageWrapper( EImageFormat::PNG );
 	if ( ImageWrapper.IsValid() && ImageWrapper->SetCompressed( RawData.GetData(), RawData.Num() ) )
 	{
 		Width = ImageWrapper->GetWidth();
diff --git a/Engine/Plugins/Editor/MobileLauncherProfileWizard/Source/MobileLauncherProfileWizard/Private/AndroidProfileWizard.h b/Engine/Plugins/Editor/MobileLauncherProfileWizard/Source/MobileLauncherProfileWizard/Private/AndroidProfileWizard.h
--- a/Engine/Plugins/Editor/MobileLauncherProfileWizard/Source/MobileLauncherProfileWizard/Private/AndroidProfileWizard.h
+++ b/Engine/Plugins/Editor/MobileLauncherProfileWizard/Source/MobileLauncherProfileWizard/Private/AndroidProfileWizard.h
@@ -2,7 +2,7 @@
 #pragma once
 
 #include "CoreMinimal.h"
-#include "Interfaces/ILauncherProfileManager.h"
+#include "ILauncherProfileManager.h"
 
 class FAndroidProfileWizard : public ILauncherProfileWizard
 {
diff --git a/Engine/Plugins/Editor/MobileLauncherProfileWizard/Source/MobileLauncherProfileWizard/Private/IOSProfileWizard.h b/Engine/Plugins/Editor/MobileLauncherProfileWizard/Source/MobileLauncherProfileWizard/Private/IOSProfileWizard.h
--- a/Engine/Plugins/Editor/MobileLauncherProfileWizard/Source/MobileLauncherProfileWizard/Private/IOSProfileWizard.h
+++ b/Engine/Plugins/Editor/MobileLauncherProfileWizard/Source/MobileLauncherProfileWizard/Private/IOSProfileWizard.h
@@ -2,7 +2,7 @@
 #pragma once
 
 #include "CoreMinimal.h"
-#include "Interfaces/ILauncherProfileManager.h"
+#include "ILauncherProfileManager.h"
 
 class FIOSProfileWizard : public ILauncherProfileWizard
 {
diff --git a/Engine/Plugins/Editor/MobileLauncherProfileWizard/Source/MobileLauncherProfileWizard/Private/MobileLauncherProfileWizard.cpp b/Engine/Plugins/Editor/MobileLauncherProfileWizard/Source/MobileLauncherProfileWizard/Private/MobileLauncherProfileWizard.cpp
--- a/Engine/Plugins/Editor/MobileLauncherProfileWizard/Source/MobileLauncherProfileWizard/Private/MobileLauncherProfileWizard.cpp
+++ b/Engine/Plugins/Editor/MobileLauncherProfileWizard/Source/MobileLauncherProfileWizard/Private/MobileLauncherProfileWizard.cpp
@@ -2,8 +2,9 @@
 
 #include "CoreMinimal.h"
 #include "Modules/ModuleManager.h"
+#include "ILauncherProfileManager.h"
+#include "ILauncherServicesModule.h"
 #include "IMobileLauncherProfileWizard.h"
-#include "Interfaces/ILauncherServicesModule.h"
 #include "AndroidProfileWizard.h"
 #include "IOSProfileWizard.h"
 
@@ -31,7 +32,7 @@ void FMobileLauncherProfileWizardModule::StartupModule()
 	ILauncherServicesModule* LauncherServicesModule = FModuleManager::GetModulePtr<ILauncherServicesModule>("LauncherServices");
 	if (LauncherServicesModule)
 	{
-		ILauncherProfileManagerRef LauncherProfileManager = LauncherServicesModule->GetProfileManager();
+		TSharedRef<ILauncherProfileManager> LauncherProfileManager = LauncherServicesModule->GetProfileManager();
 		OnProfileManagerInitialized(LauncherProfileManager.Get());
 	}
 }
@@ -43,7 +44,7 @@ void FMobileLauncherProfileWizardModule::ShutdownModule()
 	ILauncherServicesModule* LauncherServicesModule = FModuleManager::GetModulePtr<ILauncherServicesModule>("LauncherServices");
 	if (LauncherServicesModule)
 	{
-		ILauncherProfileManagerRef LauncherProfileManager = LauncherServicesModule->GetProfileManager();
+		TSharedRef<ILauncherProfileManager> LauncherProfileManager = LauncherServicesModule->GetProfileManager();
 		LauncherProfileManager->UnregisterProfileWizard(AndroidWizardPtr);
 		LauncherProfileManager->UnregisterProfileWizard(IOSWizardPtr);
 	}
diff --git a/Engine/Plugins/Editor/MobileLauncherProfileWizard/Source/MobileLauncherProfileWizard/Private/SProfileWizardUI.cpp b/Engine/Plugins/Editor/MobileLauncherProfileWizard/Source/MobileLauncherProfileWizard/Private/SProfileWizardUI.cpp
--- a/Engine/Plugins/Editor/MobileLauncherProfileWizard/Source/MobileLauncherProfileWizard/Private/SProfileWizardUI.cpp
+++ b/Engine/Plugins/Editor/MobileLauncherProfileWizard/Source/MobileLauncherProfileWizard/Private/SProfileWizardUI.cpp
@@ -29,7 +29,7 @@
 #include "InstalledPlatformInfo.h"
 #include "PlatformInfo.h"
 #include "DesktopPlatformModule.h"
-#include "Common/GameProjectHelper.h"
+#include "GameProjectHelper.h"
 
 #define LOCTEXT_NAMESPACE "MobileLauncherProfileWizard"
 
diff --git a/Engine/Plugins/Experimental/AppleARKit/Source/AppleARKit/Public/AppleARKitCameraTexture.h b/Engine/Plugins/Experimental/AppleARKit/Source/AppleARKit/Public/AppleARKitCameraTexture.h
--- a/Engine/Plugins/Experimental/AppleARKit/Source/AppleARKit/Public/AppleARKitCameraTexture.h
+++ b/Engine/Plugins/Experimental/AppleARKit/Source/AppleARKit/Public/AppleARKitCameraTexture.h
@@ -15,7 +15,7 @@ public:
 
 	// UTexture Interface.
 	APPLEARKIT_API virtual FTextureResource* CreateResource() override;
-	APPLEARKIT_API virtual EMaterialValueType GetMaterialType() override { return MCT_Texture2D; }
+	APPLEARKIT_API virtual EMaterialValueType GetMaterialType() const override { return MCT_Texture2D; }
 	APPLEARKIT_API virtual float GetSurfaceWidth() const override;
 	APPLEARKIT_API virtual float GetSurfaceHeight() const override;
 };
diff --git a/Engine/Plugins/Experimental/ImagePlate/Source/ImagePlate/Private/ImagePlateFileSequence.cpp b/Engine/Plugins/Experimental/ImagePlate/Source/ImagePlate/Private/ImagePlateFileSequence.cpp
--- a/Engine/Plugins/Experimental/ImagePlate/Source/ImagePlate/Private/ImagePlateFileSequence.cpp
+++ b/Engine/Plugins/Experimental/ImagePlate/Source/ImagePlate/Private/ImagePlateFileSequence.cpp
@@ -1,6 +1,7 @@
 // Copyright 1998-2016 Epic Games, Inc. All Rights Reserved.
 
 #include "ImagePlateFileSequence.h"
+#include "IImageWrapper.h"
 #include "IImageWrapperModule.h"
 #include "ModuleManager.h"
 #include "FileHelper.h"
@@ -537,8 +538,8 @@ namespace ImagePlateFrameCache
 
 		IImageWrapperModule& ImageWrapperModule = FModuleManager::GetModuleChecked<IImageWrapperModule>("ImageWrapper");
 
-		EImageFormat::Type ImageType = ImageWrapperModule.DetectImageFormat(SourceFileData.GetData(), SourceFileData.Num());
-		IImageWrapperPtr ImageWrapper = ImageWrapperModule.CreateImageWrapper(ImageType);
+		EImageFormat ImageType = ImageWrapperModule.DetectImageFormat(SourceFileData.GetData(), SourceFileData.Num());
+		TSharedPtr<IImageWrapper> ImageWrapper = ImageWrapperModule.CreateImageWrapper(ImageType);
 
 		if (!ImageWrapper.IsValid())
 		{
diff --git a/Engine/Plugins/Experimental/LiveLink/Source/LiveLink/Private/LiveLinkMessageBusSource.cpp b/Engine/Plugins/Experimental/LiveLink/Source/LiveLink/Private/LiveLinkMessageBusSource.cpp
--- a/Engine/Plugins/Experimental/LiveLink/Source/LiveLink/Private/LiveLinkMessageBusSource.cpp
+++ b/Engine/Plugins/Experimental/LiveLink/Source/LiveLink/Private/LiveLinkMessageBusSource.cpp
@@ -41,12 +41,12 @@ bool FLiveLinkMessageBusSource::IsSourceStillValid()
 	return true;
 }
 
-void FLiveLinkMessageBusSource::HandleHeartbeat(const FLiveLinkHeartbeatMessage& Message, const IMessageContextRef& Context)
+void FLiveLinkMessageBusSource::HandleHeartbeat(const FLiveLinkHeartbeatMessage& Message, const TSharedRef<IMessageContext, ESPMode::ThreadSafe>& Context)
 {
 	ConnectionLastActive = FPlatformTime::Seconds();
 }
 
-void FLiveLinkMessageBusSource::HandleClearSubject(const FLiveLinkClearSubject& Message, const IMessageContextRef& Context)
+void FLiveLinkMessageBusSource::HandleClearSubject(const FLiveLinkClearSubject& Message, const TSharedRef<IMessageContext, ESPMode::ThreadSafe>& Context)
 {
 	ConnectionLastActive = FPlatformTime::Seconds();
 	Client->ClearSubject(Message.SubjectName);
@@ -58,7 +58,7 @@ bool FLiveLinkMessageBusSource::RequestSourceShutdown()
 	return true;
 }
 
-void FLiveLinkMessageBusSource::HandleSubjectData(const FLiveLinkSubjectDataMessage& Message, const IMessageContextRef& Context)
+void FLiveLinkMessageBusSource::HandleSubjectData(const FLiveLinkSubjectDataMessage& Message, const TSharedRef<IMessageContext, ESPMode::ThreadSafe>& Context)
 {
 	ConnectionLastActive = FPlatformTime::Seconds();
 
@@ -91,7 +91,7 @@ void FLiveLinkMessageBusSource::HandleSubjectData(const FLiveLinkSubjectDataMess
 	Client->PushSubjectSkeleton(Message.SubjectName, Message.RefSkeleton);
 }
 
-void FLiveLinkMessageBusSource::HandleSubjectFrame(const FLiveLinkSubjectFrameMessage& Message, const IMessageContextRef& Context)
+void FLiveLinkMessageBusSource::HandleSubjectFrame(const FLiveLinkSubjectFrameMessage& Message, const TSharedRef<IMessageContext, ESPMode::ThreadSafe>& Context)
 {
 	ConnectionLastActive = FPlatformTime::Seconds();
 
diff --git a/Engine/Plugins/Experimental/LiveLink/Source/LiveLink/Private/LiveLinkMessageBusSourceEditor.cpp b/Engine/Plugins/Experimental/LiveLink/Source/LiveLink/Private/LiveLinkMessageBusSourceEditor.cpp
--- a/Engine/Plugins/Experimental/LiveLink/Source/LiveLink/Private/LiveLinkMessageBusSourceEditor.cpp
+++ b/Engine/Plugins/Experimental/LiveLink/Source/LiveLink/Private/LiveLinkMessageBusSourceEditor.cpp
@@ -108,7 +108,7 @@ void SLiveLinkMessageBusSourceEditor::OnSourceListSelectionChanged(FProviderPoll
 	SelectedResult = PollResult;
 }
 
-void SLiveLinkMessageBusSourceEditor::HandlePongMessage(const FLiveLinkPongMessage& Message, const IMessageContextRef& Context)
+void SLiveLinkMessageBusSourceEditor::HandlePongMessage(const FLiveLinkPongMessage& Message, const TSharedRef<IMessageContext, ESPMode::ThreadSafe>& Context)
 {
 	if(Message.PollRequest == CurrentPollRequest)
 	{
diff --git a/Engine/Plugins/Experimental/LiveLink/Source/LiveLink/Private/LiveLinkMessageBusSourceEditor.h b/Engine/Plugins/Experimental/LiveLink/Source/LiveLink/Private/LiveLinkMessageBusSourceEditor.h
--- a/Engine/Plugins/Experimental/LiveLink/Source/LiveLink/Private/LiveLinkMessageBusSourceEditor.h
+++ b/Engine/Plugins/Experimental/LiveLink/Source/LiveLink/Private/LiveLinkMessageBusSourceEditor.h
@@ -47,7 +47,7 @@ class SLiveLinkMessageBusSourceEditor : public SCompoundWidget
 
 private:
 
-	void HandlePongMessage(const FLiveLinkPongMessage& Message, const IMessageContextRef& Context);
+	void HandlePongMessage(const FLiveLinkPongMessage& Message, const TSharedRef<IMessageContext, ESPMode::ThreadSafe>& Context);
 
 	TSharedRef<ITableRow> MakeSourceListViewWidget(FProviderPollResultPtr PollResult, const TSharedRef<STableViewBase>& OwnerTable) const;
 
@@ -59,7 +59,7 @@ private:
 
 	FProviderPollResultPtr SelectedResult;
 
-	FMessageEndpointPtr MessageEndpoint;
+	TSharedPtr<FMessageEndpoint, ESPMode::ThreadSafe> MessageEndpoint;
 
 	FGuid CurrentPollRequest;
 };
\ No newline at end of file
diff --git a/Engine/Plugins/Experimental/LiveLink/Source/LiveLink/Public/LiveLinkMessageBusSource.h b/Engine/Plugins/Experimental/LiveLink/Source/LiveLink/Public/LiveLinkMessageBusSource.h
--- a/Engine/Plugins/Experimental/LiveLink/Source/LiveLink/Public/LiveLinkMessageBusSource.h
+++ b/Engine/Plugins/Experimental/LiveLink/Source/LiveLink/Public/LiveLinkMessageBusSource.h
@@ -40,10 +40,10 @@ public:
 private:
 
 	// Message bus message handlers
-	void HandleSubjectData(const FLiveLinkSubjectDataMessage& Message, const IMessageContextRef& Context);
-	void HandleSubjectFrame(const FLiveLinkSubjectFrameMessage& Message, const IMessageContextRef& Context);
-	void HandleHeartbeat(const FLiveLinkHeartbeatMessage& Message, const IMessageContextRef& Context);
-	void HandleClearSubject(const FLiveLinkClearSubject& Message, const IMessageContextRef& Context);
+	void HandleSubjectData(const FLiveLinkSubjectDataMessage& Message, const TSharedRef<IMessageContext, ESPMode::ThreadSafe>& Context);
+	void HandleSubjectFrame(const FLiveLinkSubjectFrameMessage& Message, const TSharedRef<IMessageContext, ESPMode::ThreadSafe>& Context);
+	void HandleHeartbeat(const FLiveLinkHeartbeatMessage& Message, const TSharedRef<IMessageContext, ESPMode::ThreadSafe>& Context);
+	void HandleClearSubject(const FLiveLinkClearSubject& Message, const TSharedRef<IMessageContext, ESPMode::ThreadSafe>& Context);
 	// End Message bus message handlers
 
 	ILiveLinkClient* Client;
@@ -51,7 +51,7 @@ private:
 	// Our identifier in LiveLink
 	FGuid SourceGuid;
 
-	FMessageEndpointPtr MessageEndpoint;
+	TSharedPtr<FMessageEndpoint, ESPMode::ThreadSafe> MessageEndpoint;
 
 	FMessageAddress ConnectionAddress;
 
diff --git a/Engine/Plugins/Experimental/StereoPanorama/Source/StereoPanorama/Private/SceneCapturer.cpp b/Engine/Plugins/Experimental/StereoPanorama/Source/StereoPanorama/Private/SceneCapturer.cpp
--- a/Engine/Plugins/Experimental/StereoPanorama/Source/StereoPanorama/Private/SceneCapturer.cpp
+++ b/Engine/Plugins/Experimental/StereoPanorama/Source/StereoPanorama/Private/SceneCapturer.cpp
@@ -7,7 +7,8 @@
 #include "StereoCapturePawn.h"
 #include "Engine/TextureRenderTarget2D.h"
 #include "Engine/World.h"
-#include "Interfaces/IImageWrapperModule.h"
+#include "IImageWrapper.h"
+#include "IImageWrapperModule.h"
 #include "Modules/ModuleManager.h"
 #include "Misc/Paths.h"
 #include "UnrealEngine.h"
@@ -648,7 +649,7 @@ TArray<FColor> USceneCapturer::SaveAtlas(FString Folder, const TArray<FColor>& S
 
 	// Write out PNG
     //TODO: ikrimae: Use threads to write out the images for performance
-	IImageWrapperPtr ImageWrapper = ImageWrapperModule.CreateImageWrapper( EImageFormat::PNG );
+	TSharedPtr<IImageWrapper> ImageWrapper = ImageWrapperModule.CreateImageWrapper( EImageFormat::PNG );
     ImageWrapper->SetRaw(SphericalAtlas.GetData(), SphericalAtlas.GetAllocatedSize(), SphericalAtlasWidth, SphericalAtlasHeight, ERGBFormat::BGRA, 8);
 	const TArray<uint8>& PNGData = ImageWrapper->GetCompressed(100);
 	FFileHelper::SaveArrayToFile( PNGData, *AtlasName );
@@ -706,7 +707,7 @@ void USceneCapturer::CaptureComponent( int32 CurrentHorizontalStep, int32 Curren
         if (FStereoPanoramaManager::GenerateDebugImages->GetInt() == 2)
         {
             //Read Whole Capture Buffer
-		    IImageWrapperPtr ImageWrapper = ImageWrapperModule.CreateImageWrapper( EImageFormat::PNG );
+		    TSharedPtr<IImageWrapper> ImageWrapper = ImageWrapperModule.CreateImageWrapper( EImageFormat::PNG );
 
             TArray<FColor> SurfaceDataWhole;
             SurfaceDataWhole.AddUninitialized(CaptureWidth * CaptureHeight);
@@ -739,7 +740,7 @@ void USceneCapturer::CaptureComponent( int32 CurrentHorizontalStep, int32 Curren
                 }
             }
 
-            IImageWrapperPtr ImageWrapper = ImageWrapperModule.CreateImageWrapper(EImageFormat::PNG);
+            TSharedPtr<IImageWrapper> ImageWrapper = ImageWrapperModule.CreateImageWrapper(EImageFormat::PNG);
             ImageWrapper->SetRaw(SurfaceData.GetData(), SurfaceData.GetAllocatedSize(), StripWidth, StripHeight, ERGBFormat::BGRA, 8);
 		    const TArray<uint8>& PNGData = ImageWrapper->GetCompressed(100);
 
diff --git a/Engine/Plugins/Media/AndroidCamera/Source/AndroidCamera/AndroidCamera.Build.cs b/Engine/Plugins/Media/AndroidCamera/Source/AndroidCamera/AndroidCamera.Build.cs
new file mode 100644
--- /dev/null
+++ b/Engine/Plugins/Media/AndroidCamera/Source/AndroidCamera/AndroidCamera.Build.cs
@@ -0,0 +1,43 @@
+// Copyright 1998-2017 Epic Games, Inc. All Rights Reserved.
+
+using System;
+using System.IO;
+
+namespace UnrealBuildTool.Rules
+{
+	public class AndroidCamera : ModuleRules
+	{
+		public AndroidCamera(ReadOnlyTargetRules Target) : base(Target)
+		{
+			DynamicallyLoadedModuleNames.AddRange(
+				new string[] {
+					"Media",
+				});
+
+			PrivateDependencyModuleNames.AddRange(
+				new string[] {
+					"AndroidCameraFactory",
+					"Core",
+					"Engine",
+					"Launch",
+					"MediaUtils",
+					"RenderCore",
+				});
+
+			PrivateIncludePathModuleNames.AddRange(
+				new string[] {
+					"Media",
+				});
+
+			PrivateIncludePaths.AddRange(
+				new string[] {
+					"AndroidCamera/Private",
+					"AndroidCamera/Private/Player",
+					"AndroidCamera/Private/Shared",
+				});
+
+			string PluginPath = Utils.MakePathRelativeTo(ModuleDirectory, Target.RelativeEnginePath);
+			AdditionalPropertiesForReceipt.Add(new ReceiptProperty("AndroidPlugin", Path.Combine(PluginPath, "AndroidCamera_UPL.xml")));
+		}
+	}
+}
