Fix Android startup movies not always playing
#ue4
#android

diff --git a/Engine/Build/Android/Java/src/com/epicgames/ue4/MediaPlayer14.java b/Engine/Build/Android/Java/src/com/epicgames/ue4/MediaPlayer14.java
--- a/Engine/Build/Android/Java/src/com/epicgames/ue4/MediaPlayer14.java
+++ b/Engine/Build/Android/Java/src/com/epicgames/ue4/MediaPlayer14.java
@@ -418,6 +418,34 @@ public class MediaPlayer14
 		return -1;
 	}
 
+	public void prepare() throws IOException, IllegalStateException
+	{
+		synchronized(this)
+		{
+			Completed = false;
+			try
+			{
+				super.prepare();
+			}
+			catch (IOException e)
+			{
+				GameActivity.Log.debug("MediaPlayer14: Prepare IOException: " + e.toString());
+				throw e;
+			}
+			catch (IllegalStateException e)
+			{
+				GameActivity.Log.debug("MediaPlayer14: Prepare IllegalStateExecption: " + e.toString());
+				throw e;
+			}
+			catch (Exception e)
+			{
+				GameActivity.Log.debug("MediaPlayer14: Prepare Exception: " + e.toString());
+				throw e;
+			}
+			Prepared = true;
+		}
+	}
+
 	public void start()
 	{
 		synchronized(this)
