UE-39451 Web browser widget causes app to crash when packaging for Distribution on Android

#UE-39451
#Android
#4.30

from CL 4067204 on Dev-Mobile
Adding a proguard exception for the WebViewControl.FrameUpdateInfo

diff --git a/Engine/Build/Android/Java/proguard-project.txt b/Engine/Build/Android/Java/proguard-project.txt
--- a/Engine/Build/Android/Java/proguard-project.txt
+++ b/Engine/Build/Android/Java/proguard-project.txt
@@ -87,6 +87,10 @@
 	public *;
 }
 
+-keep class com.epicgames.ue4.WebViewControl$FrameUpdateInfo {
+	public *;
+}
+
 # If your project uses WebView with JS, uncomment the following
 # and specify the fully qualified class name to the JavaScript interface
 # class:
diff --git a/Engine/Build/Android/Java/src/com/epicgames/ue4/WebViewControl.java b/Engine/Build/Android/Java/src/com/epicgames/ue4/WebViewControl.java
--- a/Engine/Build/Android/Java/src/com/epicgames/ue4/WebViewControl.java
+++ b/Engine/Build/Android/Java/src/com/epicgames/ue4/WebViewControl.java
@@ -79,13 +79,13 @@ class WebViewControl
 
 	public class FrameUpdateInfo 
 	{
-		java.nio.Buffer Buffer;
-		boolean FrameReady;
-		boolean RegionChanged;
-		float UScale;
-		float UOffset;
-		float VScale;
-		float VOffset;
+		public java.nio.Buffer Buffer;
+		public boolean FrameReady;
+		public boolean RegionChanged;
+		public float UScale;
+		public float UOffset;
+		public float VScale;
+		public float VOffset;
 	}
 
 	public WebViewControl(long inNativePtr, int width, int height, boolean swizzlePixels, boolean vulkanRenderer, final boolean bEnableRemoteDebugging, final boolean bUseTransparency)
