UE-59698 Cannot load youtube URLs on Android

#Android
#4.20

The native WebView cannot load HTML5 movies if the control's layout doesn't have the FILL_PARENT param

diff --git a/Engine/Build/Android/Java/src/com/epicgames/ue4/WebViewControl.java b/Engine/Build/Android/Java/src/com/epicgames/ue4/WebViewControl.java
--- a/Engine/Build/Android/Java/src/com/epicgames/ue4/WebViewControl.java
+++ b/Engine/Build/Android/Java/src/com/epicgames/ue4/WebViewControl.java
@@ -142,7 +142,9 @@ class WebViewControl
 
 				// Wrap the webview in a layout that will do absolute positioning for us
 				positionLayout = new WebViewPositionLayout(GameActivity._activity, w);
-				positionLayout.addView(webView);
+				ViewGroup.LayoutParams params = new ViewGroup.LayoutParams(ViewGroup.LayoutParams.FILL_PARENT, ViewGroup.LayoutParams.FILL_PARENT);
+
+				positionLayout.addView(webView, params);
 
 				bShown = false;
 				NextURL = null;
@@ -307,7 +309,7 @@ class WebViewControl
 					GameActivity._activity.addContentView(positionLayout, params);
 					if(!webView.IsAndroid3DBrowser)
 					{
-						GameActivity.Log.warn("request focus create");
+						//GameActivity.Log.warn("request focus create");
 						webView.requestFocus();
 					}
 				}
@@ -336,7 +338,7 @@ class WebViewControl
 						else
 						if(!webView.IsAndroid3DBrowser)
 						{
-							GameActivity.Log.warn("request focus");
+							//GameActivity.Log.warn("request focus");
 							webView.requestFocus();
 						}
 
