Add XBox 360 wired controller support for Android
#ue4
#android

diff --git a/Engine/Build/Android/Java/src/com/epicgames/ue4/GameActivity.java.template b/Engine/Build/Android/Java/src/com/epicgames/ue4/GameActivity.java.template
--- a/Engine/Build/Android/Java/src/com/epicgames/ue4/GameActivity.java.template
+++ b/Engine/Build/Android/Java/src/com/epicgames/ue4/GameActivity.java.template
@@ -2671,6 +2671,7 @@ public class GameActivity extends NativeActivity implements SurfaceHolder.Callba
 		new DeviceInfoData(0x1949, 0x0406, "Amazon Fire Game Controller"),
 		new DeviceInfoData(0x0738, 0x5263, "Mad Catz C.T.R.L.R (Smart)"),
 		new DeviceInfoData(0x0738, 0x5266, "Mad Catz C.T.R.L.R"),
+		new DeviceInfoData(0x045e, 0x028e, "Xbox Wired Controller"),
 		new DeviceInfoData(0x045e, 0x02e0, "Xbox Wireless Controller"),
 		new DeviceInfoData(0x0111, 0x1419, "SteelSeries Stratus XL"), 
 		new DeviceInfoData(0x054c, 0x05c4, "PS4 Wireless Controller")
@@ -2816,10 +2817,7 @@ public class GameActivity extends NativeActivity implements SurfaceHolder.Callba
 
 	public String getConnectionType()
 	{
-		if (isAirplaneModeOn(getApplicationContext()))
-		{
-			return "AirplaneMode";
-		}
+		boolean bAirplaneMode = isAirplaneModeOn(getApplicationContext());
 
 		ConnectivityManager cm = (ConnectivityManager)getSystemService(Context.CONNECTIVITY_SERVICE);
 		NetworkInfo activeNetwork = cm.getActiveNetworkInfo();
@@ -2831,21 +2829,18 @@ public class GameActivity extends NativeActivity implements SurfaceHolder.Callba
 			switch (connType)
 			{
 				case ConnectivityManager.TYPE_WIFI: return "WiFi";
-				case ConnectivityManager.TYPE_BLUETOOTH: return "Bluetooth";
+				case ConnectivityManager.TYPE_BLUETOOTH: return bAirplaneMode ? "AirplaneMode" : "Bluetooth";
 				case ConnectivityManager.TYPE_ETHERNET: return "Ethernet";
-				case ConnectivityManager.TYPE_WIMAX: return "WiMAX";
+				case ConnectivityManager.TYPE_WIMAX: return bAirplaneMode ? "AirplaneMode" : "WiMAX";
 			}
-			return "Cell";
+			return bAirplaneMode ? "AirplaneMode" : "Cell";
 		}
-		return "None";
+		return bAirplaneMode ? "AirplaneMode" : "None";
     }
 
 	public int AndroidThunkJava_GetNetworkConnectionType()
 	{
-		if (isAirplaneModeOn(getApplicationContext()))
-		{
-			return CONNECTION_TYPE_AIRPLANEMODE;
-		}
+		boolean bAirplaneMode = isAirplaneModeOn(getApplicationContext());
 
 		ConnectivityManager cm = (ConnectivityManager)getSystemService(Context.CONNECTIVITY_SERVICE);
 		NetworkInfo activeNetwork = cm.getActiveNetworkInfo();
@@ -2857,13 +2852,13 @@ public class GameActivity extends NativeActivity implements SurfaceHolder.Callba
 			switch (connType)
 			{
 				case ConnectivityManager.TYPE_WIFI: return CONNECTION_TYPE_WIFI;
-				case ConnectivityManager.TYPE_BLUETOOTH: return CONNECTION_TYPE_BLUETOOTH;
+				case ConnectivityManager.TYPE_BLUETOOTH: return bAirplaneMode ? CONNECTION_TYPE_AIRPLANEMODE : CONNECTION_TYPE_BLUETOOTH;
 				case ConnectivityManager.TYPE_ETHERNET: return CONNECTION_TYPE_ETHERNET;
-				case ConnectivityManager.TYPE_WIMAX: return CONNECTION_TYPE_WIMAX;
+				case ConnectivityManager.TYPE_WIMAX: return bAirplaneMode ? CONNECTION_TYPE_AIRPLANEMODE : CONNECTION_TYPE_WIMAX;
 			}
-			return CONNECTION_TYPE_CELL;
+			return bAirplaneMode ? CONNECTION_TYPE_AIRPLANEMODE : CONNECTION_TYPE_CELL;
 		}
-		return CONNECTION_TYPE_NONE;
+		return bAirplaneMode ? CONNECTION_TYPE_AIRPLANEMODE : CONNECTION_TYPE_NONE;
 	}
 	
 	public boolean AndroidThunkJava_HasMetaDataKey(String key)
@@ -3393,18 +3388,40 @@ public class GameActivity extends NativeActivity implements SurfaceHolder.Callba
 		return false;
 	}
 	
-
 	public boolean AndroidThunkJava_LaunchPackage(String packageName, String extraKey, String extraValue)
 	{
-		Intent launchIntent = getPackageManager().getLaunchIntentForPackage(packageName);
-		if (launchIntent != null)
+		Intent intent;
+		
+		int intentActionIndex = packageName.indexOf(":");
+		if (intentActionIndex >= 0)
+		{
+			String intentAction = packageName.substring(intentActionIndex+1);
+			packageName = packageName.substring(0, intentActionIndex);
+
+			intent = new Intent(intentAction);
+			if (intent == null)
+			{
+				return false;
+			}
+			intent.setPackage(packageName);
+			intent.setFlags(Intent.FLAG_ACTIVITY_NEW_TASK);
+		}
+		else
+		{
+			intent = getPackageManager().getLaunchIntentForPackage(packageName);
+		}
+
+		if (intent != null)
 		{
 			if (extraKey != "")
 			{
-				launchIntent.putExtra(extraKey, extraValue);
+				intent.putExtra(extraKey, extraValue);
+			}
+			if (intent.resolveActivity(getPackageManager()) != null)
+			{
+				startActivity(intent);
+				System.exit(0);
 			}
-			startActivity(launchIntent);
-			System.exit(0);
 		}
 		return false;
 	}
diff --git a/Engine/Source/Runtime/ApplicationCore/Private/Android/AndroidInputInterface.cpp b/Engine/Source/Runtime/ApplicationCore/Private/Android/AndroidInputInterface.cpp
--- a/Engine/Source/Runtime/ApplicationCore/Private/Android/AndroidInputInterface.cpp
+++ b/Engine/Source/Runtime/ApplicationCore/Private/Android/AndroidInputInterface.cpp
@@ -816,6 +816,10 @@ void FAndroidInputInterface::SendControllerEvents()
 						{
 							CurrentDevice.bSupportsHat = true;
 						}
+						else if (CurrentDevice.DeviceInfo.Name.StartsWith(TEXT("Xbox Wired Controller")))
+						{
+							CurrentDevice.bSupportsHat = true;
+						}
 						else if (CurrentDevice.DeviceInfo.Name.StartsWith(TEXT("Xbox Wireless Controller")))
 						{
 							CurrentDevice.bSupportsHat = true;
