Fix issue with Samsung Game Tools floating keyboard

#ue4
#android

diff --git a/Engine/Build/Android/Java/src/com/epicgames/ue4/GameActivity.java.template b/Engine/Build/Android/Java/src/com/epicgames/ue4/GameActivity.java.template
--- a/Engine/Build/Android/Java/src/com/epicgames/ue4/GameActivity.java.template
+++ b/Engine/Build/Android/Java/src/com/epicgames/ue4/GameActivity.java.template
@@ -218,6 +218,7 @@ public class GameActivity extends NativeActivity implements SurfaceHolder.Callba
 	public static VirtualKeyboardCommand lastVirtualKeyboardCommand = VirtualKeyboardCommand.VK_CMD_NONE;
 	public static final int lastVirtualKeyboardCommandDelay = 200;
 	public static final int checkLastVirtualKeyboardCommandDelay = 500;
+	public static final boolean bSamsungDevice = android.os.Build.MANUFACTURER.equals("samsung");
 	
 	static GameActivity _activity;
 	static Bundle _bundle;
@@ -1089,7 +1090,9 @@ public class GameActivity extends NativeActivity implements SurfaceHolder.Callba
 
         			nativeVirtualKeyboardShown( keyboardRect.left, keyboardRect.top, keyboardRect.right, keyboardRect.bottom );
 					//Log.debug("VK: show?" + visibleScreenYOffset + "," + newVirtualKeyboardInput.getY());
-                    if(visibleScreenYOffset > 200)
+
+					// Samsung Game Tools floating keyboard is only 1 pixel tall
+					if(visibleScreenYOffset > 200 || (bSamsungDevice && visibleScreenYOffset > 0))
                     {
 						//Log.debug("VK: show");
 						//#jira UE-55117 Android virtual keyboard can have text input hidden by software buttons
