UE-60226 Web Browser widget is completely white on Android

#Android
#4.21

from Release-4.20 CL 4122319

The URL from the sample app is malformed (http:www.google.com).
Since the phone's default browser is able to correct this problem for http and https URLs, I've added this exception to the code that checks the URL prefix, re-writing the fix for UE-53247.

diff --git a/Engine/Build/Android/Java/src/com/epicgames/ue4/WebViewControl.java b/Engine/Build/Android/Java/src/com/epicgames/ue4/WebViewControl.java
--- a/Engine/Build/Android/Java/src/com/epicgames/ue4/WebViewControl.java
+++ b/Engine/Build/Android/Java/src/com/epicgames/ue4/WebViewControl.java
@@ -79,13 +79,13 @@ class WebViewControl
 
 	public class FrameUpdateInfo 
 	{
-		java.nio.Buffer Buffer;
-		boolean FrameReady;
-		boolean RegionChanged;
-		float UScale;
-		float UOffset;
-		float VScale;
-		float VOffset;
+		public java.nio.Buffer Buffer;
+		public boolean FrameReady;
+		public boolean RegionChanged;
+		public float UScale;
+		public float UOffset;
+		public float VScale;
+		public float VOffset;
 	}
 
 	public WebViewControl(long inNativePtr, int width, int height, boolean swizzlePixels, boolean vulkanRenderer, final boolean bEnableRemoteDebugging, final boolean bUseTransparency)
@@ -142,7 +142,9 @@ class WebViewControl
 
 				// Wrap the webview in a layout that will do absolute positioning for us
 				positionLayout = new WebViewPositionLayout(GameActivity._activity, w);
-				positionLayout.addView(webView);
+				ViewGroup.LayoutParams params = new ViewGroup.LayoutParams(ViewGroup.LayoutParams.FILL_PARENT, ViewGroup.LayoutParams.FILL_PARENT);
+
+				positionLayout.addView(webView, params);
 
 				bShown = false;
 				NextURL = null;
@@ -307,7 +309,7 @@ class WebViewControl
 					GameActivity._activity.addContentView(positionLayout, params);
 					if(!webView.IsAndroid3DBrowser)
 					{
-						GameActivity.Log.warn("request focus create");
+						//GameActivity.Log.warn("request focus create");
 						webView.requestFocus();
 					}
 				}
@@ -324,7 +326,37 @@ class WebViewControl
 						else
 						if(NextURL != null)
 						{
-							if(!NextURL.contains("://") && !NextURL.startsWith("about:"))
+
+							int colPos = NextURL.indexOf(':');
+
+							boolean bNeedsPrefix = colPos < 0;;
+
+							if(!bNeedsPrefix  && !NextURL.equalsIgnoreCase("about:blank"))
+							{
+								try
+								{
+									String UrlAddress = NextURL.substring(colPos + 1);
+
+									//check if the address contains only numbers
+									bNeedsPrefix = UrlAddress.matches("[0-9]+"); // it's a port number, and URLs like "google.com:80" also need the "http://" prefix
+
+									//try to correct malformed protocols, like "http:www.google.com"
+									if(!bNeedsPrefix)
+									{
+										String UrlProtocol = NextURL.substring(0, colPos);
+																			
+										if((NextURL.equalsIgnoreCase("http") || NextURL.equalsIgnoreCase("https")) && !UrlAddress.startsWith("/"))
+										{
+											NextURL = UrlProtocol + "://" + UrlAddress;
+										}
+									}
+
+								}
+								catch(IndexOutOfBoundsException e)
+								{}
+							}
+
+							if(bNeedsPrefix)
 							{
 								//default scheme is http://
 								NextURL = "http://" + NextURL;
@@ -336,7 +368,7 @@ class WebViewControl
 						else
 						if(!webView.IsAndroid3DBrowser)
 						{
-							GameActivity.Log.warn("request focus");
+							//GameActivity.Log.warn("request focus");
 							webView.requestFocus();
 						}
 
