Fix string formatting for some locales in GetMetaDataString Method (contributed by [REDACTED_PERSON])

#PR #5026
#ue4
#android

diff --git a/Engine/Build/Android/Java/src/com/epicgames/ue4/GameActivity.java.template b/Engine/Build/Android/Java/src/com/epicgames/ue4/GameActivity.java.template
--- a/Engine/Build/Android/Java/src/com/epicgames/ue4/GameActivity.java.template
+++ b/Engine/Build/Android/Java/src/com/epicgames/ue4/GameActivity.java.template
@@ -14,6 +14,7 @@ import java.util.Map;
 import java.util.HashMap;
 import java.util.ArrayList;
 import java.text.DecimalFormat;
+import java.text.DecimalFormatSymbols;
 
 import android.annotation.TargetApi;
 
@@ -2973,7 +2974,12 @@ public class GameActivity extends NativeActivity implements SurfaceHolder.Callba
 				// note: not available so get what we can
 				getWindowManager().getDefaultDisplay().getMetrics(metrics);
 			}
-			String screenDPI = new DecimalFormat("###.##").format(metrics.xdpi) + "," + new DecimalFormat("###.##").format(metrics.ydpi) + "," + new DecimalFormat("###.##").format(metrics.densityDpi);
+
+			DecimalFormatSymbols symbols = new DecimalFormatSymbols();
+			symbols.setDecimalSeparator('.');
+			DecimalFormat decimalFormat = new DecimalFormat("###.##", symbols);
+
+			String screenDPI = decimalFormat.format(metrics.xdpi) + "," + decimalFormat.format(metrics.ydpi) + "," + decimalFormat.format(metrics.densityDpi);
 			return screenDPI;
 		}
 		else
