PR #5371: Our yearly update copyright notices to 2019 :) (Contributed by [REDACTED_PERSON])

#rnx

diff --git a/Engine/Build/Android/Java/src/com/epicgames/ue4/GameActivity.java.template b/Engine/Build/Android/Java/src/com/epicgames/ue4/GameActivity.java.template
--- a/Engine/Build/Android/Java/src/com/epicgames/ue4/GameActivity.java.template
+++ b/Engine/Build/Android/Java/src/com/epicgames/ue4/GameActivity.java.template
@@ -1,4 +1,4 @@
-// Copyright 1998-2018 Epic Games, Inc. All Rights Reserved.
+// Copyright 1998-2019 Epic Games, Inc. All Rights Reserved.
 //This file needs to be here so the "ant" build step doesnt fail when looking for a /src folder.
 
 package com.epicgames.ue4;
diff --git a/Engine/Build/Target.cs.template b/Engine/Build/Target.cs.template
--- a/Engine/Build/Target.cs.template
+++ b/Engine/Build/Target.cs.template
@@ -1,4 +1,4 @@
-// Copyright 1998-2018 Epic Games, Inc. All Rights Reserved.
+// Copyright 1998-2019 Epic Games, Inc. All Rights Reserved.
 
 using UnrealBuildTool;
 using System.Collections.Generic;
diff --git a/Engine/Shaders/Private/GpuSkinCommon.ush b/Engine/Shaders/Private/GpuSkinCommon.ush
--- a/Engine/Shaders/Private/GpuSkinCommon.ush
+++ b/Engine/Shaders/Private/GpuSkinCommon.ush
@@ -1,4 +1,4 @@
-//// Copyright 1998-2019 Epic Games, Inc. All Rights Reserved.
+// Copyright 1998-2019 Epic Games, Inc. All Rights Reserved.
 
 /*=============================================================================
 	GpuSkinCacheCommon.usf: Common properties and defines for GPU Skinning Cache
