Better chipset identification checking in ConfigRules (ignore case)

#android

diff --git a/Engine/Build/Android/Java/src/com/epicgames/ue4/GameActivity.java.template b/Engine/Build/Android/Java/src/com/epicgames/ue4/GameActivity.java.template
--- a/Engine/Build/Android/Java/src/com/epicgames/ue4/GameActivity.java.template
+++ b/Engine/Build/Android/Java/src/com/epicgames/ue4/GameActivity.java.template
@@ -1414,6 +1414,9 @@ public class GameActivity extends $${gameActivitySuperClass}$$ implements Surfac
 		// process the rules
 		if (reader != null)
 		{
+			// convert to lowercase once to check against chipset key
+			String hardwareCheck = hardware.toLowerCase();
+			
 			try
 			{
 				String line;
@@ -1490,7 +1493,7 @@ public class GameActivity extends $${gameActivitySuperClass}$$ implements Surfac
 							if (values.size() == 7)
 							{
 								String key = RemoveSurrounds(values.get(0), "\"", "\"");
-								if (hardware.contains(key))
+								if (hardwareCheck.contains(key.toLowerCase()))
 								{
 									variables.put("useAffinity", RemoveSurrounds(values.get(1), "\"", "\""));
 									variables.put("chipset", RemoveSurrounds(values.get(2), "\"", "\""));
