Enable intercepting Java side logcat on Android

#ue4
#android

diff --git a/Engine/Build/Android/Java/src/com/epicgames/ue4/GameActivity.java.template b/Engine/Build/Android/Java/src/com/epicgames/ue4/GameActivity.java.template
--- a/Engine/Build/Android/Java/src/com/epicgames/ue4/GameActivity.java.template
+++ b/Engine/Build/Android/Java/src/com/epicgames/ue4/GameActivity.java.template
@@ -2043,7 +2043,7 @@ public class GameActivity extends $${gameActivitySuperClass}$$ implements Surfac
 			Log.debug("Target SDK is " + targetSdkVersion + ".  This may cause issues if permissions are denied by the user." );				
 		}
 
-
+		Logger.RegisterCallback(this);
 		
 		// Suppress java logs in Shipping builds
 		if (nativeIsShippingBuild())
