Allow actions to be inserted whenever software keyboard opened or closed on Android

#ue4
#android

diff --git a/Engine/Build/Android/Java/src/com/epicgames/ue4/GameActivity.java.template b/Engine/Build/Android/Java/src/com/epicgames/ue4/GameActivity.java.template
--- a/Engine/Build/Android/Java/src/com/epicgames/ue4/GameActivity.java.template
+++ b/Engine/Build/Android/Java/src/com/epicgames/ue4/GameActivity.java.template
@@ -2807,7 +2807,7 @@ public class GameActivity extends $${gameActivitySuperClass}$$ implements Surfac
         	public void onGlobalLayout()
         	{
 				//Log.debug("VK: onGlobalLayout " + bKeyboardShowing);
-        		if( bKeyboardShowing )
+        		//if( bKeyboardShowing )
         		{
         			Rect visibleRect = new Rect();
         			View visibleView = mainView.getRootView();
@@ -2841,26 +2841,43 @@ public class GameActivity extends $${gameActivitySuperClass}$$ implements Surfac
 
                     int visibleScreenYOffset = Math.max(bottomDiff, topDiff);
 
-        			nativeVirtualKeyboardShown( keyboardRect.left, keyboardRect.top, keyboardRect.right, keyboardRect.bottom );
 					//Log.debug("VK: show?" + visibleScreenYOffset + "," + newVirtualKeyboardInput.getY());
 
 					// Samsung Game Tools floating keyboard is only 1 pixel tall
 					if(visibleScreenYOffset > 200 || (bSamsungDevice && visibleScreenYOffset > 0))
                     {
-						//Log.debug("VK: show");
-						//#jira UE-55117 Android virtual keyboard can have text input hidden by software buttons
-						newVirtualKeyboardInput.getLayoutParams().width = Math.abs(visibleRect.right - visibleRect.left );
-						newVirtualKeyboardInput.setX(leftDiff);
-                    	newVirtualKeyboardInput.setY(keyboardYPos);
-                    	newVirtualKeyboardInput.setVisibility(View.VISIBLE);
-						newVirtualKeyboardInput.requestFocus();
+						if (bKeyboardShowing)
+						{
+							nativeVirtualKeyboardShown( keyboardRect.left, keyboardRect.top, keyboardRect.right, keyboardRect.bottom );
+
+							//Log.debug("VK: show");
+							//#jira UE-55117 Android virtual keyboard can have text input hidden by software buttons
+							newVirtualKeyboardInput.getLayoutParams().width = Math.abs(visibleRect.right - visibleRect.left);
+							newVirtualKeyboardInput.setX(leftDiff);
+							newVirtualKeyboardInput.setY(keyboardYPos);
+							newVirtualKeyboardInput.setVisibility(View.VISIBLE);
+							newVirtualKeyboardInput.requestFocus();
+						}
+
+						AndroidThunkJava_ResizeKeyboard(true, keyboardRect, visibleRect);
                     }
-                    else if(newVirtualKeyboardInput.getY() > 0)
+                    else
                     {
-						//Log.debug("VK: hide");
-						newVirtualKeyboardInput.setVisibility(View.GONE);
-						//set offscreen
-        				newVirtualKeyboardInput.setY(-1000); 
+						if (bKeyboardShowing)
+						{
+							nativeVirtualKeyboardShown( keyboardRect.left, keyboardRect.top, keyboardRect.right, keyboardRect.bottom );
+							
+							if (newVirtualKeyboardInput.getY() > 0)
+							{
+
+								//Log.debug("VK: hide");
+								newVirtualKeyboardInput.setVisibility(View.GONE);
+								//set offscreen
+								newVirtualKeyboardInput.setY(-1000);
+							}
+						}
+
+						AndroidThunkJava_ResizeKeyboard(false, keyboardRect, visibleRect);
                     }
         		}
         	}
@@ -2875,6 +2892,11 @@ public class GameActivity extends $${gameActivitySuperClass}$$ implements Surfac
 		Log.debug("==============> GameActive.onCreate complete!");
 	}
 
+	public void AndroidThunkJava_ResizeKeyboard(boolean bShow, Rect keyboardRect, Rect visibleRect)
+	{
+		//$${gameActivityResizeKeyboardAdditions}$$
+	}
+
 	@Override
 	public void onAccuracyChanged(Sensor sensor, int accuracy)
 	{
diff --git a/Engine/Source/Programs/UnrealBuildTool/Platform/Android/UEDeployAndroid.cs b/Engine/Source/Programs/UnrealBuildTool/Platform/Android/UEDeployAndroid.cs
--- a/Engine/Source/Programs/UnrealBuildTool/Platform/Android/UEDeployAndroid.cs
+++ b/Engine/Source/Programs/UnrealBuildTool/Platform/Android/UEDeployAndroid.cs
@@ -4284,6 +4284,7 @@ namespace UnrealBuildTool
   				{ "//$${gameActivityFinalizeConfigRulesAdditions}$$", UPL.ProcessPluginNode(NDKArch, "gameActivityFinalizeConfigRulesAdditions", "")},
 				{ "//$${gameActivityBeforeConfigRulesAppliedAdditions}$$", UPL.ProcessPluginNode(NDKArch, "gameActivityBeforeConfigRulesAppliedAdditions", "")},
 				{ "//$${gameActivityAfterMainViewCreatedAdditions}$$", UPL.ProcessPluginNode(NDKArch, "gameActivityAfterMainViewCreatedAdditions", "")},
+				{ "//$${gameActivityResizeKeyboardAdditions}$$", UPL.ProcessPluginNode(NDKArch, "gameActivityResizeKeyboardAdditions", "")},
 				{ "//$${gameActivityLoggerCallbackAdditions}$$", UPL.ProcessPluginNode(NDKArch, "gameActivityLoggerCallbackAdditions", "")},
 				{ "//$${soLoadLibrary}$$", UPL.ProcessPluginNode(NDKArch, "soLoadLibrary", LoadLibraryDefaults)},
 				{ "$${gameActivitySuperClass}$$", SuperClassDefault},
