Merging CL 5808949 - support streaming live video on Android

diff --git a/Engine/Build/Android/Java/proguard-project.txt b/Engine/Build/Android/Java/proguard-project.txt
--- a/Engine/Build/Android/Java/proguard-project.txt
+++ b/Engine/Build/Android/Java/proguard-project.txt
@@ -22,6 +22,12 @@
 -keep class com.epicgames.ue4.MediaPlayer14 {
 	public *;
 }
+-keep class com.epicgames.ue4.VideoDecoder {
+	public *;
+}
+-keep class com.epicgames.ue4.AudioDecoder {
+	public *;
+}
 -keep class com.epicgames.ue4.MessageBox01 {
 	public *;
 }
diff --git a/Engine/Build/Android/Java/src/com/epicgames/ue4/AudioDecoder.java b/Engine/Build/Android/Java/src/com/epicgames/ue4/AudioDecoder.java
new file mode 100644
--- /dev/null
+++ b/Engine/Build/Android/Java/src/com/epicgames/ue4/AudioDecoder.java
@@ -0,0 +1,143 @@
+// Copyright 1998-2019 Epic Games, Inc. All Rights Reserved.
+
+package com.epicgames.ue4;
+
+import android.content.res.AssetFileDescriptor;
+import android.content.res.AssetManager;
+import android.os.Build;
+import java.io.IOException;
+import java.io.RandomAccessFile;
+import java.util.ArrayList;
+import java.util.Random;
+import android.util.Log;
+
+import android.media.MediaCodec;
+import android.media.MediaCodecInfo;
+import android.media.MediaFormat;
+import java.nio.ByteBuffer;
+
+/*
+	Custom media player that renders video to a frame buffer. This
+	variation is for API 14 and above.
+*/
+public class AudioDecoder
+{
+	private MediaCodec mAudioCodec = null;
+	private boolean bIsInitialized = false;
+	private int mSampleRate;
+	private int mChannelCount;
+	private int mBitsPerSample; 
+	private byte[] mAacSpecificConfig;
+
+	public AudioDecoder()
+	{
+	}
+	
+	public byte[] DecodeAudio(byte[] rawBuffer)
+	{
+		byte[] data = null;
+		if( mAudioCodec != null)
+		{
+			int inIndex = mAudioCodec.dequeueInputBuffer(-1);
+			if (inIndex >= 0) 
+			{
+				int sampleSize = rawBuffer.length;
+				ByteBuffer buffer = mAudioCodec.getInputBuffer(inIndex);
+				buffer.clear();
+				buffer.put(rawBuffer);
+				buffer.clear();
+
+				mAudioCodec.queueInputBuffer(inIndex, 0, sampleSize, 0, 0);
+			}
+
+			MediaCodec.BufferInfo info = new MediaCodec.BufferInfo();
+
+			int outIndex = mAudioCodec.dequeueOutputBuffer(info, 10000);
+
+			if (outIndex >= 0) 
+			{
+				ByteBuffer outBuffer = mAudioCodec.getOutputBuffer(outIndex);
+				data = new byte[info.size];
+			
+				outBuffer.get(data);
+				outBuffer.clear();
+				
+				mAudioCodec.releaseOutputBuffer(outIndex, true);
+			}
+		}
+		return data;
+	}
+
+	public void reset()
+	{
+		release();
+	}
+
+	public void release()
+	{
+		if (null != mAudioCodec)
+		{
+			mAudioCodec.stop();
+			mAudioCodec.release();
+		}
+	}
+
+	private boolean CreateCodec()
+	{
+		try 
+		{
+			// Create the decoder
+			mAudioCodec = MediaCodec.createDecoderByType(MediaFormat.MIMETYPE_AUDIO_AAC);
+
+			// Media settings
+			MediaFormat format = MediaFormat.createAudioFormat(MediaFormat.MIMETYPE_AUDIO_AAC, mSampleRate, mChannelCount);
+
+			// PPS and SPS headers
+			format.setInteger(MediaFormat.KEY_BIT_RATE, mBitsPerSample);
+
+			byte[] bytes = new byte[]{(byte) 0x12, (byte)0x12};
+			ByteBuffer bb = ByteBuffer.wrap(bytes);
+			format.setByteBuffer("csd-0", bb);
+			
+			// Configure the decoder
+			mAudioCodec.configure(format, null, null, 0);
+
+			// Start the decoder
+			mAudioCodec.start();
+			bIsInitialized = true;
+		}
+		catch(Exception e)
+		{
+			GameActivity.Log.warn("Android Audio Decoder: CreateCodec failed!");		
+			e.printStackTrace();
+			return false;
+		}
+		
+		return true;
+	}
+
+	public boolean updateConfig(int sampleRate, int channelCount, int BitsPerSample, byte[] AacSpecificConfig)
+	{
+		bIsInitialized = false;
+		
+		mSampleRate = sampleRate;
+		mChannelCount = channelCount;
+		mBitsPerSample = BitsPerSample;
+		mAacSpecificConfig = AacSpecificConfig;
+
+		if(mAudioCodec != null)
+		{
+			try 
+			{
+				mAudioCodec.stop();
+				mAudioCodec.release();
+			}
+			catch(Exception e)
+			{
+				e.printStackTrace();
+			}
+		}
+		return CreateCodec();
+	}
+
+}
diff --git a/Engine/Build/Android/Java/src/com/epicgames/ue4/VideoDecoder.java b/Engine/Build/Android/Java/src/com/epicgames/ue4/VideoDecoder.java
new file mode 100644
--- /dev/null
+++ b/Engine/Build/Android/Java/src/com/epicgames/ue4/VideoDecoder.java
@@ -0,0 +1,1479 @@
+// Copyright 1998-2019 Epic Games, Inc. All Rights Reserved.
+
+package com.epicgames.ue4;
+
+import android.content.res.AssetFileDescriptor;
+import android.content.res.AssetManager;
+import android.graphics.Rect;
+import android.opengl.*;
+import android.os.Build;
+import java.io.File;
+import java.io.IOException;
+import java.io.RandomAccessFile;
+import android.media.MediaExtractor;
+import android.media.MediaFormat;
+import android.media.MediaPlayer;
+import android.media.MediaPlayer.TrackInfo;
+import java.util.ArrayList;
+import java.util.Random;
+import android.util.Log;
+import java.net.URL;
+import java.net.HttpURLConnection;
+
+import android.graphics.SurfaceTexture;
+import android.media.MediaCodec;
+import android.media.MediaCodecInfo;
+import android.media.MediaFormat;
+import android.view.Surface;
+import android.view.TextureView;
+import java.nio.ByteBuffer;
+
+import java.util.List;
+import java.util.ListIterator;
+import java.util.Map.Entry;
+import java.util.LinkedList;
+import java.util.AbstractMap.SimpleEntry;
+
+/*
+	Custom media player that renders video to a frame buffer. This
+	variation is for API 14 and above.
+*/
+public class VideoDecoder
+{
+	private boolean SwizzlePixels = true;
+	private boolean VulkanRenderer = false;
+	private boolean SupportsImageExternal = false;
+	private int FrameWidth = 0;
+	private int FrameHeight = 0;
+
+	private volatile boolean WaitOnBitmapRender = false;
+
+	private BitmapRenderer mBitmapRenderer = null;
+	private OESTextureRenderer mOESTextureRenderer = null;
+	private MediaCodec mVideoCodec = null;// Video decoder
+	private MediaCodec mAudioCodec = null;// Video decoder
+	private Surface mVideoCodecSurface = null;
+	private boolean isEOS = false;
+	private boolean bIsInitialized = false;
+
+	private volatile boolean NewFrameAvailable = true;
+
+	private byte[] SPS;
+	private byte[] PPS;
+
+	public class FrameUpdateInfo 
+	{
+		public java.nio.Buffer Buffer;
+		public double Timestamp;
+		public double Duration;
+		public boolean FrameReady;
+		public boolean RegionChanged;
+		public float UScale;
+		public float UOffset;
+		public float VScale;
+		public float VOffset;
+	}
+	
+	// ======================================================================================
+
+	public VideoDecoder(boolean swizzlePixels, boolean vulkanRenderer, boolean supportsImageExternal)
+	{
+		SwizzlePixels = swizzlePixels;
+		VulkanRenderer = vulkanRenderer;
+		SupportsImageExternal = supportsImageExternal;
+		WaitOnBitmapRender = false;
+	}
+	
+	LinkedList<Entry<Double,Double>> frameTimeDurationList = new LinkedList<>();
+
+	void addFrameTimeDuration(double timestamp, double duration)
+	{
+		synchronized(this)
+		{
+			Entry<Double,Double> frameTimeDuration = new SimpleEntry<>(timestamp, duration);
+			frameTimeDurationList.add(frameTimeDuration);
+		}
+	}	
+
+	Entry<Double,Double> getFrameTimeDuration(double timestamp)
+	{
+		Entry<Double,Double> frameTimeDuration = new SimpleEntry<>(timestamp, 100.0);
+
+		synchronized(this)
+		{
+			for (ListIterator<Entry<Double,Double>> iterator = frameTimeDurationList.listIterator(); iterator.hasNext(); ) 
+			{
+			    Entry<Double,Double> element = iterator.next();
+
+				double timeStampDiff = Math.abs(frameTimeDuration.getKey() - timestamp);
+
+				// if the decoder returns 0 for the timestamp - BUG on older phones, get the first frame from the stack
+				// or check if we get the exact value from the stack
+				if(timestamp <= 0 || element.getKey() == timestamp)
+				{
+					frameTimeDuration = element;
+					break;
+				}
+				
+				// try to aproximate the closest timestamp value, by parsing the stack
+				if( timeStampDiff != 0 && timeStampDiff > Math.abs(timestamp - element.getKey()))
+				{
+					frameTimeDuration = element;
+				}
+			}
+		}
+
+		// remove the entry from the stack
+		frameTimeDurationList.remove(frameTimeDuration);
+		return frameTimeDuration;
+	}
+
+	public boolean DecodeVideo(byte[] rawBuffer, double timeStamp, double duration)
+	{
+		if(mVideoCodec != null)
+		{
+			try
+			{
+				synchronized(this)
+				{
+					int inIndex = mVideoCodec.dequeueInputBuffer(-1);
+					if (inIndex >= 0) 
+					{
+						int sampleSize = rawBuffer.length;
+				
+						ByteBuffer buffer = mVideoCodec.getInputBuffer(inIndex);
+						buffer.clear();
+						buffer.put(rawBuffer);
+						buffer.clear();
+
+						mVideoCodec.queueInputBuffer(inIndex, 0, sampleSize, (long)timeStamp, 0);
+					}
+
+					MediaCodec.BufferInfo info = new MediaCodec.BufferInfo();
+
+					int outIndex = mVideoCodec.dequeueOutputBuffer(info, -1);
+					if (outIndex >= 0) 
+					{
+						addFrameTimeDuration(timeStamp, duration); //us to ms
+						mVideoCodec.releaseOutputBuffer(outIndex, true);
+					}
+
+					if ((info.flags & MediaCodec.BUFFER_FLAG_END_OF_STREAM) != 0) 
+					{
+						GameActivity.Log.warn("OutputBuffer BUFFER_FLAG_END_OF_STREAM");
+					}
+					return true;
+				}
+			}
+			catch(Exception e)
+			{
+				e.printStackTrace();
+			}
+
+		}
+		return false;
+	}
+			
+	public boolean didResolutionChange()
+	{
+		if (null != mOESTextureRenderer)
+		{
+			return mOESTextureRenderer.resolutionChanged();
+		}
+		if (null != mBitmapRenderer)
+		{
+			return mBitmapRenderer.resolutionChanged();
+		}
+		return false;
+	}
+
+	public int getExternalTextureId()
+	{
+		if (null != mOESTextureRenderer)
+		{
+			return mOESTextureRenderer.getExternalTextureId();
+		}
+		if (null != mBitmapRenderer)
+		{
+			return mBitmapRenderer.getExternalTextureId();
+		}
+		return -1;
+	}
+
+	public void reset()
+	{
+		release();
+	}
+
+	public void release()
+	{
+		if (null != mOESTextureRenderer)
+		{
+			while (WaitOnBitmapRender) ;
+			releaseOESTextureRenderer();
+		}
+		if (null != mBitmapRenderer)
+		{
+			while (WaitOnBitmapRender) ;
+			releaseOESTextureRenderer();
+		}
+		if (null != mVideoCodec)
+		{
+			mVideoCodec.stop();
+			mVideoCodec.release();
+		}
+	}
+
+	private ByteBuffer getSPSPPSHeader()
+	{
+		byte[] header = {0, 0, 0, 1};
+		final ByteBuffer csd_info =  ByteBuffer.allocate(SPS.length + PPS.length + 8);
+		csd_info.put(header);
+		csd_info.put(SPS);
+		csd_info.put(header);
+		csd_info.put(PPS);
+		
+		return csd_info;
+	}
+	
+	private double getFrameTimestamp()
+	{
+		SurfaceTexture videoSurfaceTexture = null;
+		double retVal = -1;
+
+		if (null != mOESTextureRenderer)
+		{
+			videoSurfaceTexture = mOESTextureRenderer.getSurfaceTexture();
+		}
+		if (null != mBitmapRenderer)
+		{
+			videoSurfaceTexture = mBitmapRenderer.getSurfaceTexture();
+		}
+
+		if(videoSurfaceTexture != null)
+		{
+			retVal = videoSurfaceTexture.getTimestamp() / 1000; // nanos to micros
+		}
+		return retVal;
+	}
+
+	private boolean CreateCodec()
+	{
+		try 
+		{
+			// Create the decoder
+			mVideoCodec = MediaCodec.createDecoderByType(MediaFormat.MIMETYPE_VIDEO_AVC);
+
+			// Media settings
+			MediaFormat format = MediaFormat.createVideoFormat(MediaFormat.MIMETYPE_VIDEO_AVC, FrameWidth, FrameHeight);// MIMETYPE: a two-part identifier for file formats and format contents
+
+			// PPS and SPS headers
+			format.setByteBuffer("csd-0", getSPSPPSHeader());
+
+			int frameRate = 30; // 30 fps
+
+			format.setInteger(MediaFormat.KEY_MAX_INPUT_SIZE, FrameWidth * FrameHeight);
+			format.setInteger(MediaFormat.KEY_COLOR_FORMAT, MediaCodecInfo.CodecCapabilities.COLOR_FormatSurface);
+			format.setInteger(MediaFormat.KEY_FRAME_RATE, frameRate);
+			format.setInteger(MediaFormat.KEY_CAPTURE_RATE, frameRate);
+
+			isEOS = false;
+
+			// Configure the decoder
+			mVideoCodec.configure(format, mVideoCodecSurface, null, 0);
+
+			// Start the decoder
+			mVideoCodec.start();
+			bIsInitialized = true;
+		}
+		catch(Exception e)
+		{
+			GameActivity.Log.warn("Android Video Decoder: CreateCodec failed!");		
+			e.printStackTrace();
+			return false;
+		}
+		
+		return true;
+	}
+
+	public boolean updateConfig(int frameWidth, int frameHeight, byte[] inSPS, byte[] inPPS)
+	{
+		bIsInitialized = false;
+
+		FrameWidth	= frameWidth;
+		FrameHeight	= frameHeight;
+		SPS = inSPS;
+		PPS = inPPS;
+
+		release();
+
+		if(mVideoCodec == null)
+		{
+			if(SwizzlePixels || VulkanRenderer || !SupportsImageExternal)
+			{
+				//initBitmapRenderer();
+			}
+
+			return true;
+		}
+
+		try 
+		{
+			// Get an instance of MediaCodec and give it its Mime type
+			mVideoCodec.stop();
+			mVideoCodec.release();
+		}
+		catch(Exception e)
+		{
+			e.printStackTrace();
+		}
+
+		return CreateCodec();
+	}
+
+	// ======================================================================================
+
+	private boolean CreateBitmapRenderer()
+	{
+		releaseBitmapRenderer();
+
+		mBitmapRenderer = new BitmapRenderer(SwizzlePixels, VulkanRenderer);
+		if (!mBitmapRenderer.isValid())
+		{
+			mBitmapRenderer = null;
+			return false;
+		}
+		
+		// set this here as the size may have been set before the GL resources were created.
+		mBitmapRenderer.setSize(FrameWidth, FrameHeight);
+		mVideoCodecSurface = mBitmapRenderer.getSurface();
+		CreateCodec();
+
+		return true;
+	}
+
+	void releaseBitmapRenderer()
+	{
+		if (null != mBitmapRenderer)
+		{
+			mBitmapRenderer.release();
+			mBitmapRenderer = null;
+		}
+	}
+
+	public void initBitmapRenderer()
+	{
+		// if not already allocated.
+		// Create bitmap renderer's gl resources in the renderer thread.
+		if (null == mBitmapRenderer)
+		{
+			if (!CreateBitmapRenderer())
+			{
+				GameActivity.Log.warn("initBitmapRenderer failed to alloc mBitmapRenderer ");
+				reset();
+			  }
+		}
+	}
+
+	public FrameUpdateInfo getVideoLastFrameData()
+	{
+		initBitmapRenderer();
+		if (null != mBitmapRenderer)
+		{
+			WaitOnBitmapRender = true;
+			FrameUpdateInfo frameInfo = mBitmapRenderer.updateFrameData();
+			WaitOnBitmapRender = false;
+			return frameInfo;
+		}
+		else
+		{
+			return null;
+		}
+	}
+
+	public FrameUpdateInfo getVideoLastFrame(int destTexture)
+	{
+
+		initBitmapRenderer();
+		if (null != mBitmapRenderer)
+		{
+			WaitOnBitmapRender = true;
+			FrameUpdateInfo frameInfo = mBitmapRenderer.updateFrameData(destTexture);
+			WaitOnBitmapRender = false;
+			return frameInfo;
+		}
+		else
+		{
+			return null;
+		}
+	}
+
+	/*
+		All this internal surface view does is manage the
+		offscreen bitmap that the media player decoding can
+		render into for eventual extraction to the UE4 buffers.
+	*/
+	class BitmapRenderer
+		implements android.graphics.SurfaceTexture.OnFrameAvailableListener
+	{
+		private java.nio.Buffer mFrameData = null;
+		private android.graphics.SurfaceTexture mSurfaceTexture = null;
+		private int mTextureWidth = -1;
+		private int mTextureHeight = -1;
+		private android.view.Surface mSurface = null;
+		private boolean mFrameAvailable = false;
+		private int mTextureID = -1;
+		private int mFBO = -1;
+		private int mBlitVertexShaderID = -1;
+		private int mBlitFragmentShaderID = -1;
+		private float[] mTransformMatrix = new float[16];
+		private boolean mTriangleVerticesDirty = true;
+		private boolean mTextureSizeChanged = true;
+		private boolean mUseOwnContext = true;
+		private boolean mVulkanRenderer = false;
+		private boolean mSwizzlePixels = false;
+
+		private int GL_TEXTURE_EXTERNAL_OES = 0x8D65;
+
+		private EGLDisplay mEglDisplay;
+		private EGLContext mEglContext;
+		private EGLSurface mEglSurface;
+
+		private EGLDisplay mSavedDisplay;
+		private EGLContext mSavedContext;
+		private EGLSurface mSavedSurfaceDraw;
+		private EGLSurface mSavedSurfaceRead;
+
+		private boolean mCreatedEGLDisplay = false;
+
+		public BitmapRenderer(boolean swizzlePixels, boolean isVulkan)
+		{
+			mSwizzlePixels = swizzlePixels;
+			mVulkanRenderer = isVulkan;
+
+			mEglSurface = EGL14.EGL_NO_SURFACE;
+			mEglContext = EGL14.EGL_NO_CONTEXT;
+			mUseOwnContext = true;
+
+			if (mVulkanRenderer)
+			{
+				mSwizzlePixels = true;
+
+					GameActivity.Log.debug("BitmapRenderer: mSwizzlePixels");
+			}
+			else
+			{
+				String RendererString = GLES20.glGetString(GLES20.GL_RENDERER);
+				GameActivity.Log.debug("BitmapRenderer 2" + RendererString);
+				// Do not use shared context if Adreno before 400 or on older Android than Marshmallow
+				if (RendererString.contains("Adreno (TM) "))
+				{
+					int AdrenoVersion = Integer.parseInt(RendererString.substring(12));
+					if (AdrenoVersion < 400 || android.os.Build.VERSION.SDK_INT < 22)
+					{
+						GameActivity.Log.debug("VideoDecoder: disabled shared GL context on " + RendererString);
+						mUseOwnContext = false;
+					}
+				}
+			}
+
+			if (mUseOwnContext)
+			{
+			GameActivity.Log.debug("BitmapRenderer 3");
+				initContext();
+				GameActivity.Log.debug("BitmapRenderer 4");
+				saveContext();
+				GameActivity.Log.debug("BitmapRenderer 5");
+				makeCurrent();
+				GameActivity.Log.debug("BitmapRenderer 6");
+				initSurfaceTexture();
+				GameActivity.Log.debug("BitmapRenderer 7");
+				restoreContext();
+			}
+			else
+			{
+			GameActivity.Log.debug("BitmapRenderer 8");
+				initSurfaceTexture();
+			}
+			GameActivity.Log.debug("BitmapRenderer 9");
+		}
+
+		private void initContext()
+		{
+			mEglDisplay = EGL14.EGL_NO_DISPLAY;
+			EGLContext shareContext = EGL14.EGL_NO_CONTEXT;
+
+			int majorver[] = new int[] { 0 };
+			int minorver[] = new int[] { 0 };
+			if (!mVulkanRenderer)
+			{
+				mEglDisplay = EGL14.eglGetCurrentDisplay();
+				shareContext = EGL14.eglGetCurrentContext();
+
+				if (android.os.Build.VERSION.SDK_INT >= 18 &&
+					EGL14.eglQueryContext(mEglDisplay, shareContext, EGLExt.EGL_CONTEXT_MAJOR_VERSION_KHR, majorver, 0) &&
+					EGL14.eglQueryContext(mEglDisplay, shareContext, EGLExt.EGL_CONTEXT_MINOR_VERSION_KHR, minorver, 0))
+				{
+					GameActivity.Log.debug("VideoDecoder: Existing GL context is version " + majorver[0] + "." + minorver[0]);
+				}
+				else
+				// on some devices eg Galaxy S6, the above fails but we do get EGL14.EGL_CONTEXT_CLIENT_VERSION=3
+				if (EGL14.eglQueryContext(mEglDisplay, shareContext, EGL14.EGL_CONTEXT_CLIENT_VERSION, majorver, 0))
+				{					
+					GameActivity.Log.debug("VideoDecoder: Existing GL context is version " + majorver[0]);
+				}
+				else
+				{
+					GameActivity.Log.debug("VideoDecoder: Existing GL context version not detected");		
+				}
+			}
+			else
+			{
+				mEglDisplay = EGL14.eglGetDisplay(EGL14.EGL_DEFAULT_DISPLAY);
+				if (mEglDisplay == EGL14.EGL_NO_DISPLAY)
+				{
+					GameActivity.Log.error("unable to get EGL14 display");
+					return;
+				}
+				int[] version = new int[2];
+				if (!EGL14.eglInitialize(mEglDisplay, version, 0, version, 1))
+				{
+					mEglDisplay = null;
+					GameActivity.Log.error("unable to initialize EGL14 display");
+					return;
+				}				
+				
+				mCreatedEGLDisplay = true;
+			}
+
+			int[] configSpec = new int[]
+			{
+				EGL14.EGL_RENDERABLE_TYPE, EGL14.EGL_OPENGL_ES2_BIT,
+				EGL14.EGL_SURFACE_TYPE, EGL14.EGL_PBUFFER_BIT,
+				EGL14.EGL_NONE
+			};
+			EGLConfig[] configs = new EGLConfig[1];
+			int[] num_config = new int[1];
+			EGL14.eglChooseConfig(mEglDisplay, configSpec, 0, configs, 0, 1, num_config, 0);
+			int[] contextAttribsES2 = new int[]
+			{
+				EGL14.EGL_CONTEXT_CLIENT_VERSION, 2,
+				EGL14.EGL_NONE
+			};
+			int[] contextAttribsES31 = new int[]
+			{
+				EGLExt.EGL_CONTEXT_MAJOR_VERSION_KHR, 3,
+				EGLExt.EGL_CONTEXT_MINOR_VERSION_KHR, 1,
+				EGL14.EGL_NONE
+			};
+			mEglContext = EGL14.eglCreateContext(mEglDisplay, configs[0], shareContext, majorver[0]==3 ? contextAttribsES31 : contextAttribsES2, 0);
+
+			if (EGL14.eglQueryString(mEglDisplay, EGL14.EGL_EXTENSIONS).contains("EGL_KHR_surfaceless_context"))
+			{
+				mEglSurface = EGL14.EGL_NO_SURFACE;
+			}
+			else
+			{
+				int[] pbufferAttribs = new int[]
+				{
+					EGL14.EGL_NONE
+				};
+				mEglSurface = EGL14.eglCreatePbufferSurface(mEglDisplay, configs[0], pbufferAttribs, 0);
+			}
+		}
+
+		private void saveContext()
+		{
+			mSavedDisplay = EGL14.eglGetCurrentDisplay();
+			mSavedContext = EGL14.eglGetCurrentContext();
+			mSavedSurfaceDraw = EGL14.eglGetCurrentSurface(EGL14.EGL_DRAW);
+			mSavedSurfaceRead = EGL14.eglGetCurrentSurface(EGL14.EGL_READ);
+		}
+
+		private void makeCurrent()
+		{
+			EGL14.eglMakeCurrent(mEglDisplay, mEglSurface, mEglSurface, mEglContext);
+		}
+
+		private void restoreContext()
+		{
+			EGL14.eglMakeCurrent(mSavedDisplay, mSavedSurfaceDraw, mSavedSurfaceRead, mSavedContext);
+		}
+
+		private void initSurfaceTexture()
+		{
+			int[] textures = new int[1];
+			GLES20.glGenTextures(1, textures, 0);
+			mTextureID = textures[0];
+			if (mTextureID <= 0)
+			{
+				GameActivity.Log.error("mTextureID <= 0");
+				release();
+				return;
+			}
+			mSurfaceTexture = new android.graphics.SurfaceTexture(mTextureID);
+			mSurfaceTexture.setOnFrameAvailableListener(this);
+			mSurface = new android.view.Surface(mSurfaceTexture);
+
+			int[] glInt = new int[1];
+
+			GLES20.glGenFramebuffers(1,glInt,0);
+			mFBO = glInt[0];
+			if (mFBO <= 0)
+			{
+				GameActivity.Log.error("mFBO <= 0");
+				release();
+				return;
+			}
+			// Special shaders for blit of movie texture.
+			mBlitVertexShaderID = createShader(GLES20.GL_VERTEX_SHADER, mBlitVextexShader);
+			if (mBlitVertexShaderID == 0)
+			{
+				GameActivity.Log.error("mBlitVertexShaderID == 0");
+				release();
+				return;
+			}
+			int mBlitFragmentShaderID = createShader(GLES20.GL_FRAGMENT_SHADER,
+				mSwizzlePixels ? mBlitFragmentShaderBGRA : mBlitFragmentShaderRGBA);
+			if (mBlitFragmentShaderID == 0)
+			{
+				GameActivity.Log.error("mBlitFragmentShaderID == 0");
+				release();
+				return;
+			}
+			mProgram = GLES20.glCreateProgram();
+			if (mProgram <= 0)
+			{
+				GameActivity.Log.error("mProgram <= 0");
+				release();
+				return;
+			}
+			GLES20.glAttachShader(mProgram, mBlitVertexShaderID);
+			GLES20.glAttachShader(mProgram, mBlitFragmentShaderID);
+			GLES20.glLinkProgram(mProgram);
+			int[] linkStatus = new int[1];
+			GLES20.glGetProgramiv(mProgram, GLES20.GL_LINK_STATUS, linkStatus, 0);
+			if (linkStatus[0] != GLES20.GL_TRUE)
+			{
+				GameActivity.Log.error("Could not link program: ");
+				GameActivity.Log.error(GLES20.glGetProgramInfoLog(mProgram));
+				GLES20.glDeleteProgram(mProgram);
+				mProgram = 0;
+				release();
+				return;
+			}
+			mPositionAttrib = GLES20.glGetAttribLocation(mProgram, "Position");
+			mTexCoordsAttrib = GLES20.glGetAttribLocation(mProgram, "TexCoords");
+			mTextureUniform = GLES20.glGetUniformLocation(mProgram, "VideoTexture");
+
+			GLES20.glGenBuffers(1,glInt,0);
+			mBlitBuffer = glInt[0];
+			if (mBlitBuffer <= 0)
+			{
+				GameActivity.Log.error("mBlitBuffer <= 0");
+				release();
+				return;
+			}
+
+			// Create blit mesh.
+			mTriangleVertices = java.nio.ByteBuffer.allocateDirect(
+				mTriangleVerticesData.length * FLOAT_SIZE_BYTES)
+					.order(java.nio.ByteOrder.nativeOrder()).asFloatBuffer();
+			mTriangleVerticesDirty = true;
+
+			// Set up GL state
+			if (mUseOwnContext)
+			{
+				GLES20.glDisable(GLES20.GL_BLEND);
+				GLES20.glDisable(GLES20.GL_CULL_FACE);
+				GLES20.glDisable(GLES20.GL_SCISSOR_TEST);
+				GLES20.glDisable(GLES20.GL_STENCIL_TEST);
+				GLES20.glDisable(GLES20.GL_DEPTH_TEST);
+				GLES20.glDisable(GLES20.GL_DITHER);
+				GLES20.glColorMask(true,true,true,true);
+			}
+		}
+		
+		private void UpdateVertexData()
+		{
+			if (!mTriangleVerticesDirty || mBlitBuffer <= 0)
+			{
+				return;
+			}
+
+			// fill it in
+			mTriangleVertices.position(0);
+			mTriangleVertices.put(mTriangleVerticesData).position(0);
+
+			// save VBO state
+			int[] glInt = new int[1];
+			GLES20.glGetIntegerv(GLES20.GL_ARRAY_BUFFER_BINDING, glInt, 0);
+			int previousVBO = glInt[0];
+			
+			GLES20.glBindBuffer(GLES20.GL_ARRAY_BUFFER, mBlitBuffer);
+			GLES20.glBufferData(GLES20.GL_ARRAY_BUFFER,
+				mTriangleVerticesData.length*FLOAT_SIZE_BYTES,
+				mTriangleVertices, GLES20.GL_STATIC_DRAW);
+
+			// restore VBO state
+			GLES20.glBindBuffer(GLES20.GL_ARRAY_BUFFER, previousVBO);
+			
+			mTriangleVerticesDirty = false;
+		}
+
+		public boolean isValid()
+		{
+			return mSurfaceTexture != null;
+		}
+
+		private int createShader(int shaderType, String source)
+		{
+			int shader = GLES20.glCreateShader(shaderType);
+			if (shader != 0)
+			{
+				GLES20.glShaderSource(shader, source);
+				GLES20.glCompileShader(shader);
+				int[] compiled = new int[1];
+				GLES20.glGetShaderiv(shader, GLES20.GL_COMPILE_STATUS, compiled, 0);
+				if (compiled[0] == 0)
+				{
+					GameActivity.Log.error("Could not compile shader " + shaderType + ":");
+					GameActivity.Log.error(GLES20.glGetShaderInfoLog(shader));
+					GLES20.glDeleteShader(shader);
+					shader = 0;
+				}
+			}
+			return shader;
+		}
+
+		public void onFrameAvailable(android.graphics.SurfaceTexture st)
+		{
+			synchronized(this)
+			{
+				mFrameAvailable = true;
+			}
+		}
+
+		public android.graphics.SurfaceTexture getSurfaceTexture()
+		{
+			return mSurfaceTexture;
+		}
+
+		public android.view.Surface getSurface()
+		{
+			return mSurface;
+		}
+
+		public int getExternalTextureId()
+		{
+			return mTextureID;
+		}
+
+		// NOTE: Synchronized with updateFrameData to prevent frame
+		// updates while the surface may need to get reallocated.
+		public void setSize(int width, int height)
+		{
+			synchronized(this)
+			{
+				if (width != mTextureWidth ||
+					height != mTextureHeight)
+				{
+					mTextureWidth = width;
+					mTextureHeight = height;
+					mFrameData = null;
+					mTextureSizeChanged = true;
+				}
+			}
+		}
+
+		public boolean resolutionChanged()
+		{
+			boolean changed;
+			synchronized(this)
+			{
+				changed = mTextureSizeChanged;
+				mTextureSizeChanged = false;
+			}
+			return changed;
+		}
+
+		private static final int FLOAT_SIZE_BYTES = 4;
+		private static final int TRIANGLE_VERTICES_DATA_STRIDE_BYTES = 4 * FLOAT_SIZE_BYTES;
+		private static final int TRIANGLE_VERTICES_DATA_POS_OFFSET = 0;
+		private static final int TRIANGLE_VERTICES_DATA_UV_OFFSET = 2;
+		private float[] mTriangleVerticesData = {
+			// X, Y, U, V
+			-1.0f, -1.0f, 0.f, 0.f,
+			1.0f, -1.0f, 1.f, 0.f,
+			-1.0f, 1.0f, 0.f, 1.f,
+			1.0f, 1.0f, 1.f, 1.f,
+			};
+
+		private java.nio.FloatBuffer mTriangleVertices;
+
+		private final String mBlitVextexShader =
+			"attribute vec2 Position;\n" +
+			"attribute vec2 TexCoords;\n" +
+			"varying vec2 TexCoord;\n" +
+			"void main() {\n" +
+			"	TexCoord = TexCoords;\n" +
+			"	gl_Position = vec4(Position, 0.0, 1.0);\n" +
+			"}\n";
+
+		// NOTE: We read the fragment as BGRA so that in the end, when
+		// we glReadPixels out of the FBO, we get them in that order
+		// and avoid having to swizzle the pixels in the CPU.
+		private final String mBlitFragmentShaderBGRA =
+			"#extension GL_OES_EGL_image_external : require\n" +
+			"uniform samplerExternalOES VideoTexture;\n" +
+			"varying highp vec2 TexCoord;\n" +
+			"void main()\n" +
+			"{\n" +
+			"	gl_FragColor = texture2D(VideoTexture, TexCoord).bgra;\n" +
+			"}\n";
+		private final String mBlitFragmentShaderRGBA =
+			"#extension GL_OES_EGL_image_external : require\n" +
+			"uniform samplerExternalOES VideoTexture;\n" +
+			"varying highp vec2 TexCoord;\n" +
+			"void main()\n" +
+			"{\n" +
+			"	gl_FragColor = texture2D(VideoTexture, TexCoord).rgba;\n" +
+			"}\n";
+
+		private int mProgram;
+		private int mPositionAttrib;
+		private int mTexCoordsAttrib;
+		private int mBlitBuffer;
+		private int mTextureUniform;
+
+		public FrameUpdateInfo updateFrameData()
+		{
+			synchronized(this)
+			{
+				if (null == mFrameData && mTextureWidth > 0 && mTextureHeight > 0)
+				{
+					mFrameData = java.nio.ByteBuffer.allocateDirect(mTextureWidth*mTextureHeight*4);
+				}
+				// Copy surface texture to frame data.
+				if (!copyFrameTexture(0, mFrameData))
+				{
+					return null;
+				}
+				FrameUpdateInfo frameUpdateInfo = new FrameUpdateInfo();
+
+				frameUpdateInfo.Buffer = mFrameData;
+
+				Entry<Double,Double> frameTimeDuration = getFrameTimeDuration(getFrameTimestamp());
+				frameUpdateInfo.Timestamp = frameTimeDuration.getKey();
+				frameUpdateInfo.Duration = frameTimeDuration.getValue();
+				frameUpdateInfo.FrameReady = true;
+				frameUpdateInfo.RegionChanged = false;
+
+				return frameUpdateInfo;
+			}
+
+		}
+
+		public FrameUpdateInfo updateFrameData(int destTexture)
+		{
+			synchronized(this)
+			{
+				// Copy surface texture to destination texture.
+				if (!copyFrameTexture(destTexture, null))
+				{
+					return null;
+				}
+			
+				FrameUpdateInfo frameUpdateInfo = new FrameUpdateInfo();
+
+				frameUpdateInfo.Buffer = null;
+
+				Entry<Double,Double> frameTimeDuration = getFrameTimeDuration(getFrameTimestamp());
+				frameUpdateInfo.Timestamp = frameTimeDuration.getKey();
+				frameUpdateInfo.Duration = frameTimeDuration.getValue();
+
+				frameUpdateInfo.FrameReady = true;
+				frameUpdateInfo.RegionChanged = false;
+
+				return frameUpdateInfo;
+			}
+		}
+
+		// Copy the surface texture to another texture, or to raw data.
+		// Note: copying to raw data creates a temporary FBO texture.
+		private boolean copyFrameTexture(int destTexture, java.nio.Buffer destData)
+		{
+			if (!mFrameAvailable)
+			{
+				// We only return fresh data when we generate it. At other
+				// time we return nothing to indicate that there was nothing
+				// new to return. The media player deals with this by keeping
+				// the last frame around and using that for rendering.
+				return false;
+			}
+			mFrameAvailable = false;
+			if (null == mSurfaceTexture)
+			{
+				// Can't update if there's no surface to update into.
+				return false;
+			}
+
+			int[] glInt = new int[1];
+			boolean[] glBool = new boolean[1];
+			// Either use own context or save states
+			boolean previousBlend=false, previousCullFace=false, previousScissorTest=false, previousStencilTest=false, previousDepthTest=false, previousDither=false;
+			int previousFBO=0, previousVBO=0, previousMinFilter=0, previousMagFilter=0;
+			int[] previousViewport = new int[4];
+			if (mUseOwnContext)
+			{
+				// Received reports of these not being preserved when changing contexts
+				GLES20.glActiveTexture(GLES20.GL_TEXTURE0);
+				GLES20.glGetTexParameteriv(GLES20.GL_TEXTURE_2D, GLES20.GL_TEXTURE_MIN_FILTER, glInt, 0);
+				previousMinFilter = glInt[0];
+				GLES20.glGetTexParameteriv(GLES20.GL_TEXTURE_2D, GLES20.GL_TEXTURE_MAG_FILTER, glInt, 0);
+				previousMagFilter = glInt[0];
+
+				saveContext();
+				makeCurrent();
+			}
+			else
+			{
+				// Clear gl errors as they can creep in from the UE4 renderer.
+				GLES20.glGetError();
+
+				previousBlend = GLES20.glIsEnabled(GLES20.GL_BLEND);
+				previousCullFace = GLES20.glIsEnabled(GLES20.GL_CULL_FACE);
+				previousScissorTest = GLES20.glIsEnabled(GLES20.GL_SCISSOR_TEST);
+				previousStencilTest = GLES20.glIsEnabled(GLES20.GL_STENCIL_TEST);
+				previousDepthTest = GLES20.glIsEnabled(GLES20.GL_DEPTH_TEST);
+				previousDither = GLES20.glIsEnabled(GLES20.GL_DITHER);
+				GLES20.glGetIntegerv(GLES20.GL_FRAMEBUFFER_BINDING, glInt, 0);
+				previousFBO = glInt[0];
+				GLES20.glGetIntegerv(GLES20.GL_ARRAY_BUFFER_BINDING, glInt, 0);
+				previousVBO = glInt[0];
+				GLES20.glGetIntegerv(GLES20.GL_VIEWPORT, previousViewport, 0);
+
+				GLES20.glActiveTexture(GLES20.GL_TEXTURE0);
+				GLES20.glGetTexParameteriv(GLES20.GL_TEXTURE_2D, GLES20.GL_TEXTURE_MIN_FILTER, glInt, 0);
+				previousMinFilter = glInt[0];
+				GLES20.glGetTexParameteriv(GLES20.GL_TEXTURE_2D, GLES20.GL_TEXTURE_MAG_FILTER, glInt, 0);
+				previousMagFilter = glInt[0];
+
+				glVerify("save state");
+			}
+
+			// Get the latest video texture frame.
+			mSurfaceTexture.updateTexImage();
+			mSurfaceTexture.getTransformMatrix(mTransformMatrix);
+				
+			float UMin = mTransformMatrix[12];
+			float UMax = UMin + mTransformMatrix[0];
+			float VMin = mTransformMatrix[13];
+			float VMax = VMin + mTransformMatrix[5];
+				
+			if (mTriangleVerticesData[2] != UMin ||
+				mTriangleVerticesData[6] != UMax ||
+				mTriangleVerticesData[11] != VMin ||
+				mTriangleVerticesData[3] != VMax)
+			{
+				//GameActivity.Log.debug("Matrix:");
+				//GameActivity.Log.debug(mTransformMatrix[0] + " " + mTransformMatrix[1] + " " + mTransformMatrix[2] + " " + mTransformMatrix[3]);
+				//GameActivity.Log.debug(mTransformMatrix[4] + " " + mTransformMatrix[5] + " " + mTransformMatrix[6] + " " + mTransformMatrix[7]);
+				//GameActivity.Log.debug(mTransformMatrix[8] + " " + mTransformMatrix[9] + " " + mTransformMatrix[10] + " " + mTransformMatrix[11]);
+				//GameActivity.Log.debug(mTransformMatrix[12] + " " + mTransformMatrix[13] + " " + mTransformMatrix[14] + " " + mTransformMatrix[15]);
+				mTriangleVerticesData[ 2] = mTriangleVerticesData[10] = UMin;
+				mTriangleVerticesData[ 6] = mTriangleVerticesData[14] = UMax;
+				mTriangleVerticesData[11] = mTriangleVerticesData[15] = VMin;
+				mTriangleVerticesData[ 3] = mTriangleVerticesData[ 7] = VMax;
+				mTriangleVerticesDirty = true;
+				//GameActivity.Log.debug("U = " + mTriangleVerticesData[2] + ", " + mTriangleVerticesData[6]);
+				//GameActivity.Log.debug("V = " + mTriangleVerticesData[11] + ", " + mTriangleVerticesData[3]);
+			}
+
+			if (null != destData)
+			{
+				// Rewind data so that we can write to it.
+				destData.position(0);
+			}
+
+			if (!mUseOwnContext)
+			{
+				GLES20.glDisable(GLES20.GL_BLEND);
+				GLES20.glDisable(GLES20.GL_CULL_FACE);
+				GLES20.glDisable(GLES20.GL_SCISSOR_TEST);
+				GLES20.glDisable(GLES20.GL_STENCIL_TEST);
+				GLES20.glDisable(GLES20.GL_DEPTH_TEST);
+				GLES20.glDisable(GLES20.GL_DITHER);
+				GLES20.glColorMask(true,true,true,true);
+
+				glVerify("reset state");
+			}
+
+			GLES20.glViewport(0, 0, mTextureWidth, mTextureHeight);
+
+			glVerify("set viewport");
+
+			// Set-up FBO target texture..
+			int FBOTextureID = 0;
+			if (null != destData)
+			{
+				// Create temporary FBO for data copy.
+				GLES20.glGenTextures(1,glInt,0);
+				FBOTextureID = glInt[0];
+			}
+			else
+			{
+				// Use the given texture as the FBO.
+				FBOTextureID = destTexture;
+			}
+			// Set the FBO to draw into the texture one-to-one.
+			GLES20.glBindTexture(GLES20.GL_TEXTURE_2D, FBOTextureID);
+			GLES20.glTexParameteri(GLES20.GL_TEXTURE_2D,
+				GLES20.GL_TEXTURE_MIN_FILTER, GLES20.GL_NEAREST);
+			GLES20.glTexParameteri(GLES20.GL_TEXTURE_2D,
+				GLES20.GL_TEXTURE_MAG_FILTER, GLES20.GL_NEAREST);
+			GLES20.glTexParameteri(GLES20.GL_TEXTURE_2D,
+				GLES20.GL_TEXTURE_WRAP_S, GLES20.GL_CLAMP_TO_EDGE);
+			GLES20.glTexParameteri(GLES20.GL_TEXTURE_2D,
+				GLES20.GL_TEXTURE_WRAP_T, GLES20.GL_CLAMP_TO_EDGE);
+			// Create the temp FBO data if needed.
+			if (null != destData)
+			{
+				//int w = 1<<(32-Integer.numberOfLeadingZeros(mTextureWidth-1));
+				//int h = 1<<(32-Integer.numberOfLeadingZeros(mTextureHeight-1));
+				GLES20.glTexImage2D(GLES20.GL_TEXTURE_2D, 0,
+					GLES20.GL_RGBA,
+					mTextureWidth, mTextureHeight,
+					0, GLES20.GL_RGBA, GLES20.GL_UNSIGNED_BYTE, null);
+			}
+
+			glVerify("set-up FBO texture");
+			// Set to render to the FBO.
+			GLES20.glBindFramebuffer(GLES20.GL_FRAMEBUFFER, mFBO);
+			glVerify("glBindFramebuffer");
+			GLES20.glFramebufferTexture2D(
+				GLES20.GL_FRAMEBUFFER,
+				GLES20.GL_COLOR_ATTACHMENT0,
+				GLES20.GL_TEXTURE_2D, FBOTextureID, 0);
+
+			// check status
+			int status = GLES20.glCheckFramebufferStatus(GLES20.GL_FRAMEBUFFER);
+			if (status != GLES20.GL_FRAMEBUFFER_COMPLETE)
+			{
+				GameActivity.Log.warn("Failed to complete framebuffer attachment ("+status+")");
+			}
+
+			// The special shaders to render from the video texture.
+			GLES20.glUseProgram(mProgram);
+
+			// Set the mesh that renders the video texture.
+			UpdateVertexData();
+			GLES20.glBindBuffer(GLES20.GL_ARRAY_BUFFER, mBlitBuffer);
+			GLES20.glEnableVertexAttribArray(mPositionAttrib);
+			GLES20.glVertexAttribPointer(mPositionAttrib, 2, GLES20.GL_FLOAT, false,
+				TRIANGLE_VERTICES_DATA_STRIDE_BYTES, 0);
+			GLES20.glEnableVertexAttribArray(mTexCoordsAttrib);
+			GLES20.glVertexAttribPointer(mTexCoordsAttrib, 2, GLES20.GL_FLOAT, false,
+				TRIANGLE_VERTICES_DATA_STRIDE_BYTES,
+				TRIANGLE_VERTICES_DATA_UV_OFFSET*FLOAT_SIZE_BYTES);
+
+			glVerify("setup movie texture read");
+			
+			GLES20.glClearColor(1.0f, 1.0f, 1.0f, 0);
+
+			GLES20.glClear( GLES20.GL_COLOR_BUFFER_BIT);
+
+			// connect 'VideoTexture' to video source texture (mTextureID).
+			// mTextureID is bound to GL_TEXTURE_EXTERNAL_OES in updateTexImage
+			GLES20.glUniform1i(mTextureUniform, 0);
+			GLES20.glActiveTexture(GLES20.GL_TEXTURE0);
+			GLES20.glBindTexture(GL_TEXTURE_EXTERNAL_OES, mTextureID);
+
+			// Draw the video texture mesh.
+			GLES20.glDrawArrays(GLES20.GL_TRIANGLE_STRIP, 0, 4);
+
+			GLES20.glFlush();
+			
+			// Read the FBO texture pixels into raw data.
+			if (null != destData)
+			{
+				GLES20.glReadPixels(
+					0, 0, mTextureWidth, mTextureHeight,
+					GLES20.GL_RGBA, GLES20.GL_UNSIGNED_BYTE,
+					destData);
+			}
+
+			glVerify("draw & read movie texture");
+
+			// Restore state and cleanup.
+			if (mUseOwnContext)
+			{
+				GLES20.glFramebufferTexture2D(
+					GLES20.GL_FRAMEBUFFER,
+					GLES20.GL_COLOR_ATTACHMENT0,
+					GLES20.GL_TEXTURE_2D, 0, 0);
+
+				if (null != destData && FBOTextureID > 0)
+				{
+					glInt[0] = FBOTextureID;
+					GLES20.glDeleteTextures(1, glInt, 0);
+				}
+
+				restoreContext();
+
+				// Restore previous texture filtering
+				GLES20.glTexParameteri(GLES20.GL_TEXTURE_2D, GLES20.GL_TEXTURE_MIN_FILTER, previousMinFilter);
+				GLES20.glTexParameteri(GLES20.GL_TEXTURE_2D, GLES20.GL_TEXTURE_MAG_FILTER, previousMagFilter);
+			}
+			else
+			{
+				GLES20.glBindFramebuffer(GLES20.GL_FRAMEBUFFER, previousFBO);
+				if (null != destData && FBOTextureID > 0)
+				{
+					glInt[0] = FBOTextureID;
+					GLES20.glDeleteTextures(1, glInt, 0);
+				}
+				GLES20.glBindBuffer(GLES20.GL_ARRAY_BUFFER, previousVBO);
+
+				GLES20.glViewport(previousViewport[0], previousViewport[1],	previousViewport[2], previousViewport[3]);
+				if (previousBlend) GLES20.glEnable(GLES20.GL_BLEND);
+				if (previousCullFace) GLES20.glEnable(GLES20.GL_CULL_FACE);
+				if (previousScissorTest) GLES20.glEnable(GLES20.GL_SCISSOR_TEST);
+				if (previousStencilTest) GLES20.glEnable(GLES20.GL_STENCIL_TEST);
+				if (previousDepthTest) GLES20.glEnable(GLES20.GL_DEPTH_TEST);
+				if (previousDither) GLES20.glEnable(GLES20.GL_DITHER);
+
+				GLES20.glTexParameteri(GLES20.GL_TEXTURE_2D, GLES20.GL_TEXTURE_MIN_FILTER, previousMinFilter);
+				GLES20.glTexParameteri(GLES20.GL_TEXTURE_2D, GLES20.GL_TEXTURE_MAG_FILTER, previousMagFilter);
+
+				// invalidate cached state in RHI
+				GLES20.glDisableVertexAttribArray(mPositionAttrib);
+				GLES20.glDisableVertexAttribArray(mTexCoordsAttrib);
+				nativeClearCachedAttributeState(mPositionAttrib, mTexCoordsAttrib);
+			}
+
+			return true;
+		}
+
+		private void showGlError(String op, int error)
+		{
+			switch (error)
+			{
+				case GLES20.GL_INVALID_ENUM:						GameActivity.Log.error("MediaPlayer$BitmapRenderer: " + op + ": glGetError GL_INVALID_ENUM");  break;
+				case GLES20.GL_INVALID_OPERATION:					GameActivity.Log.error("MediaPlayer$BitmapRenderer: " + op + ": glGetError GL_INVALID_OPERATION");  break;
+				case GLES20.GL_INVALID_FRAMEBUFFER_OPERATION:		GameActivity.Log.error("MediaPlayer$BitmapRenderer: " + op + ": glGetError GL_INVALID_FRAMEBUFFER_OPERATION");  break;
+				case GLES20.GL_INVALID_VALUE:						GameActivity.Log.error("MediaPlayer$BitmapRenderer: " + op + ": glGetError GL_INVALID_VALUE");  break;
+				case GLES20.GL_FRAMEBUFFER_INCOMPLETE_ATTACHMENT:	GameActivity.Log.error("MediaPlayer$BitmapRenderer: " + op + ": glGetError GL_FRAMEBUFFER_INCOMPLETE_ATTACHMENT");  break;
+				case GLES20.GL_FRAMEBUFFER_INCOMPLETE_DIMENSIONS:	GameActivity.Log.error("MediaPlayer$BitmapRenderer: " + op + ": glGetError GL_FRAMEBUFFER_INCOMPLETE_DIMENSIONS");  break;
+				case GLES20.GL_FRAMEBUFFER_UNSUPPORTED:				GameActivity.Log.error("MediaPlayer$BitmapRenderer: " + op + ": glGetError GL_FRAMEBUFFER_UNSUPPORTED");  break;
+				case GLES20.GL_OUT_OF_MEMORY:						GameActivity.Log.error("MediaPlayer$BitmapRenderer: " + op + ": glGetError GL_OUT_OF_MEMORY");  break;
+				default:											GameActivity.Log.error("MediaPlayer$BitmapRenderer: " + op + ": glGetError " + error);
+			}
+		}
+
+		private void glVerify(String op)
+		{
+			int error;
+			while ((error = GLES20.glGetError()) != GLES20.GL_NO_ERROR)
+			{
+				showGlError(op, error);
+				throw new RuntimeException(op + ": glGetError " + error);
+			}
+		}
+
+		private void glWarn(String op)
+		{
+			int error;
+			while ((error = GLES20.glGetError()) != GLES20.GL_NO_ERROR)
+			{
+				showGlError(op, error);
+			}
+		}
+
+		public void release()
+		{
+			if (null != mSurface)
+			{
+				mSurface.release();
+				mSurface = null;
+			}
+			if (null != mSurfaceTexture)
+			{
+				mSurfaceTexture.release();
+				mSurfaceTexture = null;
+			}
+			int[] glInt = new int[1];
+			if (mBlitBuffer > 0)
+			{
+				glInt[0] = mBlitBuffer;
+				GLES20.glDeleteBuffers(1,glInt,0);
+				mBlitBuffer = -1;
+			}
+			if (mProgram > 0)
+			{
+				GLES20.glDeleteProgram(mProgram);
+				mProgram = -1;
+			}
+			if (mBlitVertexShaderID > 0)
+			{
+				GLES20.glDeleteShader(mBlitVertexShaderID);
+				mBlitVertexShaderID = -1;
+			}
+			if (mBlitFragmentShaderID > 0)
+			{
+				GLES20.glDeleteShader(mBlitFragmentShaderID);
+				mBlitFragmentShaderID = -1;
+			}
+			if (mFBO > 0)
+			{
+				glInt[0] = mFBO;
+				GLES20.glDeleteFramebuffers(1,glInt,0);
+				mFBO = -1;
+			}
+			if (mTextureID > 0)
+			{
+				glInt[0] = mTextureID;
+				GLES20.glDeleteTextures(1,glInt,0);
+				mTextureID = -1;
+			}
+			if (mEglSurface != EGL14.EGL_NO_SURFACE)
+			{
+				EGL14.eglDestroySurface(mEglDisplay, mEglSurface);
+				mEglSurface = EGL14.EGL_NO_SURFACE;
+			}
+			if (mEglContext != EGL14.EGL_NO_CONTEXT)
+			{
+				EGL14.eglDestroyContext(mEglDisplay, mEglContext);
+				mEglContext = EGL14.EGL_NO_CONTEXT;
+			}
+			if (mCreatedEGLDisplay)
+			{
+				EGL14.eglTerminate(mEglDisplay);
+				mEglDisplay = EGL14.EGL_NO_DISPLAY;
+				mCreatedEGLDisplay = false;	
+			}
+		}
+	};
+
+	public native void nativeClearCachedAttributeState(int PositionAttrib, int TexCoordsAttrib);
+
+	// ======================================================================================
+
+	private boolean CreateOESTextureRenderer(int OESTextureId)
+	{
+		releaseOESTextureRenderer();
+		mOESTextureRenderer = new OESTextureRenderer(OESTextureId);
+		if (!mOESTextureRenderer.isValid())
+		{
+			mOESTextureRenderer = null;
+			return false;
+		}
+		
+		// set this here as the size may have been set before the GL resources were created.
+		mOESTextureRenderer.setSize(FrameWidth, FrameHeight);
+
+		mVideoCodecSurface = mOESTextureRenderer.getSurface();
+		CreateCodec();
+
+		return true;
+	}
+
+	void releaseOESTextureRenderer()
+	{
+		if (null != mOESTextureRenderer)
+		{
+			mOESTextureRenderer.release();
+			mOESTextureRenderer = null;
+		}
+	}
+	
+	public FrameUpdateInfo updateVideoFrame(int externalTextureId)
+	{
+		if (null == mOESTextureRenderer)
+		{
+			if (!CreateOESTextureRenderer(externalTextureId))
+			{
+				GameActivity.Log.warn("updateVideoFrame failed to alloc mOESTextureRenderer ");
+				reset();
+				return null;
+			}
+		}
+		WaitOnBitmapRender = true;
+		FrameUpdateInfo result = mOESTextureRenderer.updateVideoFrame();
+		WaitOnBitmapRender = false;
+		return result;
+	}
+
+	public boolean isCodecReady()
+	{
+		return bIsInitialized;
+	}
+
+	/*
+		This handles events for our OES texture
+	*/
+	class OESTextureRenderer
+		implements android.graphics.SurfaceTexture.OnFrameAvailableListener
+	{
+		private android.graphics.SurfaceTexture mSurfaceTexture = null;
+		private int mTextureWidth = -1;
+		private int mTextureHeight = -1;
+		private android.view.Surface mSurface = null;
+		private boolean mFrameAvailable = false;
+		private int mTextureID = -1;
+		private float[] mTransformMatrix = new float[16];
+		private boolean mTextureSizeChanged = true;
+
+		private float mUScale = 1.0f;
+		private float mVScale = -1.0f;
+		private float mUOffset = 0.0f;
+		private float mVOffset = 1.0f;
+
+		public OESTextureRenderer(int OESTextureId)
+		{
+			mTextureID = OESTextureId;
+
+			mSurfaceTexture = new android.graphics.SurfaceTexture(mTextureID);
+			mSurfaceTexture.setOnFrameAvailableListener(this);
+			mSurface = new android.view.Surface(mSurfaceTexture);
+		}
+
+		public void release()
+		{
+			if (null != mSurface)
+			{
+				mSurface.release();
+				mSurface = null;
+			}
+			if (null != mSurfaceTexture)
+			{
+				mSurfaceTexture.release();
+				mSurfaceTexture = null;
+			}
+		}
+
+		public boolean isValid()
+		{
+			return mSurfaceTexture != null;
+		}
+
+		public void onFrameAvailable(android.graphics.SurfaceTexture st)
+		{
+			synchronized(this)
+			{
+				mFrameAvailable = true;
+			}
+		}
+
+		public android.graphics.SurfaceTexture getSurfaceTexture()
+		{
+			return mSurfaceTexture;
+		}
+
+		public android.view.Surface getSurface()
+		{
+			return mSurface;
+		}
+
+		public int getExternalTextureId()
+		{
+			return mTextureID;
+		}
+
+		// NOTE: Synchronized with updateFrameData to prevent frame
+		// updates while the surface may need to get reallocated.
+		public void setSize(int width, int height)
+		{
+			synchronized(this)
+			{
+				if (width != mTextureWidth ||
+					height != mTextureHeight)
+				{
+					mTextureWidth = width;
+					mTextureHeight = height;
+					mTextureSizeChanged = true;
+				}
+			}
+		}
+
+		public boolean resolutionChanged()
+		{
+			boolean changed;
+			synchronized(this)
+			{
+				changed = mTextureSizeChanged;
+				mTextureSizeChanged = false;
+			}
+			return changed;
+		}
+
+		public FrameUpdateInfo updateVideoFrame()
+		{
+			synchronized(this)
+			{
+				return getFrameUpdateInfo();
+			}
+		}
+
+		private FrameUpdateInfo getFrameUpdateInfo()
+		{
+			FrameUpdateInfo frameUpdateInfo = new FrameUpdateInfo();
+
+			frameUpdateInfo.FrameReady = false;
+			frameUpdateInfo.RegionChanged = false;
+			frameUpdateInfo.UScale = mUScale;
+			frameUpdateInfo.UOffset = mUOffset;
+
+			// note: the matrix has V flipped
+			frameUpdateInfo.VScale = -mVScale;
+			frameUpdateInfo.VOffset = 1.0f - mVOffset;
+
+			if (!mFrameAvailable)
+			{
+				// We only return fresh data when we generate it. At other
+				// time we return nothing to indicate that there was nothing
+				// new to return. The media player deals with this by keeping
+				// the last frame around and using that for rendering.
+				return frameUpdateInfo;
+			}
+
+			mFrameAvailable = false;
+			if (null == mSurfaceTexture)
+			{
+				// Can't update if there's no surface to update into.
+				return frameUpdateInfo;
+			}
+
+			frameUpdateInfo.FrameReady = true;
+
+			// Get the latest video texture frame.
+			mSurfaceTexture.updateTexImage();
+			mSurfaceTexture.getTransformMatrix(mTransformMatrix);
+
+			if (mUScale != mTransformMatrix[0] ||
+				mVScale != mTransformMatrix[5] ||
+				mUOffset != mTransformMatrix[12] ||
+				mVOffset != mTransformMatrix[13])
+			{
+				mUScale = mTransformMatrix[0];
+				mVScale = mTransformMatrix[5];
+				mUOffset = mTransformMatrix[12];
+				mVOffset = mTransformMatrix[13];
+
+				frameUpdateInfo.RegionChanged = true;
+				frameUpdateInfo.UScale = mUScale;
+				frameUpdateInfo.UOffset = mUOffset;
+
+				// note: the matrix has V flipped
+				frameUpdateInfo.VScale = -mVScale;
+				frameUpdateInfo.VOffset = 1.0f - mVOffset;
+			}
+
+			Entry<Double,Double> frameTimeDuration = getFrameTimeDuration(getFrameTimestamp());
+			frameUpdateInfo.Timestamp = frameTimeDuration.getKey();
+			frameUpdateInfo.Duration = frameTimeDuration.getValue();
+			return frameUpdateInfo;
+		}
+	};
+
+}
diff --git a/Engine/Plugins/Media/AndroidMedia/Source/AndroidMedia/Private/Player/AndroidMediaTextureSample.h b/Engine/Plugins/Media/AndroidMedia/Source/AndroidMedia/Private/Player/AndroidMediaTextureSample.h
--- a/Engine/Plugins/Media/AndroidMedia/Source/AndroidMedia/Private/Player/AndroidMediaTextureSample.h
+++ b/Engine/Plugins/Media/AndroidMedia/Source/AndroidMedia/Private/Player/AndroidMediaTextureSample.h
@@ -223,6 +223,12 @@ public:
 		return Time;
 	}
 
+	void UpdateTimeDuration(FTimespan InTime, FTimespan InDuration)
+	{
+		Time = InTime;
+		Duration = InDuration;
+	}
+
 	virtual bool IsCacheable() const override
 	{
 #if WITH_ENGINE
