Add permission validation tool for Android

#ue4
#android

diff --git a/Engine/Build/Android/Prebuilt/ValidatePermissions/classes/placeholder.txt b/Engine/Build/Android/Prebuilt/ValidatePermissions/classes/placeholder.txt
new file mode 100644
--- /dev/null
+++ b/Engine/Build/Android/Prebuilt/ValidatePermissions/classes/placeholder.txt
@@ -0,0 +1 @@
+This file is here to create the directory.
\ No newline at end of file
diff --git a/Engine/Build/Android/Prebuilt/ValidatePermissions/compile.bat b/Engine/Build/Android/Prebuilt/ValidatePermissions/compile.bat
new file mode 100644
--- /dev/null
+++ b/Engine/Build/Android/Prebuilt/ValidatePermissions/compile.bat
@@ -0,0 +1,3 @@
+%JAVA_HOME%\bin\javac.exe -d classes src\*.java
+cd classes
+%JAVA_HOME%\bin\jar cfe ../bin/ValidatePermissions.jar ValidatePermissions *.class
\ No newline at end of file
diff --git a/Engine/Build/Android/Prebuilt/ValidatePermissions/compile.sh b/Engine/Build/Android/Prebuilt/ValidatePermissions/compile.sh
new file mode 100755
--- /dev/null
+++ b/Engine/Build/Android/Prebuilt/ValidatePermissions/compile.sh
@@ -0,0 +1,4 @@
+#!/bin/bash
+javac -d classes src/*.java
+cd classes
+jar cfe ../bin/ValidatePermissions.jar ValidatePermissions *.class
diff --git a/Engine/Build/Android/Prebuilt/ValidatePermissions/src/ValidatePermissions.java b/Engine/Build/Android/Prebuilt/ValidatePermissions/src/ValidatePermissions.java
new file mode 100644
--- /dev/null
+++ b/Engine/Build/Android/Prebuilt/ValidatePermissions/src/ValidatePermissions.java
@@ -0,0 +1,144 @@
+// Copyright 1998-2019 Epic Games, Inc. All Rights Reserved.
+
+import java.io.*;
+import java.util.*;
+
+import java.nio.file.Files;
+import java.nio.file.Paths;
+import java.nio.file.Path;
+
+
+public class ValidatePermissions
+{
+	public static Map<String, String> ReadWhitelist(String filename)
+	{
+		Path path = Paths.get(filename);
+		if (Files.notExists(path))
+		{
+			return null;
+		}
+
+		Map<String, String> list = new HashMap<String, String>();
+		try
+		{
+			BufferedReader br = new BufferedReader(new FileReader(filename));
+
+			String line;
+			while ((line = br.readLine()) != null)
+			{
+				list.put(line, "");
+			}
+			return list;
+		}
+		catch (Exception e)
+		{
+		}
+		return null;
+	}
+
+	public static void main(String[] args)
+	{
+		String whitelistFilename = "permission_whitelist.txt";
+		Map<String, String> whitelist = ReadWhitelist(whitelistFilename);
+		if (whitelist == null)
+		{
+			System.out.println("ValidatePermissions: permission_whitelist.txt not provided; not validating.");
+			System.exit(0);
+		}
+
+		String aaptPath = "";
+		String apkFilename = "";
+	
+		if (args.length < 2)
+		{
+			System.out.println("ValidatePermissions v1.0\n");
+			System.out.println("Usage: aaptPath input.apk\n");
+			System.exit(0);
+		}
+
+		aaptPath = args[0];
+		apkFilename = args[1];
+
+		Path apkPath = Paths.get(apkFilename);
+		if (Files.notExists(apkPath))
+		{
+			System.out.println("ERROR: Unable to find APK: " + apkFilename);
+			System.exit(-1);
+		}
+
+		if (!System.getProperty("os.name").toLowerCase().contains("win"))
+		{
+			// remove the .exe at the end of aapt path if not WindowsBorders
+			if (aaptPath.endsWith(".exe"))
+			{
+				aaptPath = aaptPath.substring(0, aaptPath.length() - 4);
+			}
+		}
+		
+		Process process = null;
+		try
+		{
+			process = Runtime.getRuntime().exec(aaptPath + " d xmltree \"" + apkFilename + "\" AndroidManifest.xml");
+			if (process == null)
+			{
+				System.out.println("ERROR: Unable to run AAPT: " + aaptPath);
+				System.exit(-1);
+			}
+		}
+		catch (Exception e)
+		{
+			System.out.println("ERROR: Unable to run AAPT: " + aaptPath);
+			System.exit(-1);
+		}
+
+		InputStream is = process.getInputStream();  
+		InputStreamReader isr = new InputStreamReader(is);
+		BufferedReader br = new BufferedReader(isr);  
+
+		String line;
+		boolean bNextIsPermission = false;
+
+		try
+		{
+			while ((line = br.readLine()) != null)
+			{
+				if (bNextIsPermission)
+				{
+					bNextIsPermission = false;
+
+					int startIndex = line.indexOf('"');
+					if (startIndex < 0)
+					{
+						continue;
+					}
+					line = line.substring(startIndex + 1);
+					int stopIndex = line.indexOf('"');
+					if (stopIndex < 0)
+					{
+						continue;
+					}
+					line = line.substring(0, stopIndex);
+
+					// check against whitelist
+					if (!whitelist.containsKey(line))
+					{
+						System.out.println("ERROR: Permission '" + line + "' NOT whitelisted!");
+						System.exit(-1);
+					}
+				}
+				else if (line.contains("uses-permission"))
+				{
+					bNextIsPermission = true;
+				}
+			}
+		}
+		catch (Exception e)
+		{
+			System.out.println("Error running AAPT: " + aaptPath);
+			System.exit(-1);
+		}
+
+		System.out.println("Permissions successfully validated.");
+		System.exit(0);
+	}
+}
\ No newline at end of file
