Only keep Profile override from ConifgRules if OverrideProfile left set

#android

diff --git a/Engine/Build/Android/Java/src/com/epicgames/ue4/GameActivity.java.template b/Engine/Build/Android/Java/src/com/epicgames/ue4/GameActivity.java.template
--- a/Engine/Build/Android/Java/src/com/epicgames/ue4/GameActivity.java.template
+++ b/Engine/Build/Android/Java/src/com/epicgames/ue4/GameActivity.java.template
@@ -1549,6 +1549,7 @@ public class GameActivity extends $${gameActivitySuperClass}$$ implements Surfac
 		variables.put("memory", Integer.toString(totalMemory));
 
 		//variables.put("Profile", "Android");
+		variables.put("OverrideProfile", "true");
 		variables.put("hardware", hardware);
 		variables.put("ro.hardware", roHardware);
 		variables.put("processor", processor);
@@ -2054,6 +2055,15 @@ public class GameActivity extends $${gameActivitySuperClass}$$ implements Surfac
 			}
 		}
 
+		// clear Profile if OverrideProfile was not left set as well
+		if (variables.containsKey("Profile"))
+		{
+			if (!variables.containsKey("OverrideProfile"))
+			{
+				variables.remove("Profile");
+			}
+		}
+
 		// notify AndroidProcess of the core masks and affinity enable state
 		boolean bUseAffinity = variables.get("useAffinity").equals("true");
 		int littleCoreMask = Integer.decode(variables.get("littleCoreMask"));
