Fix for Android crash when playing streaming media

diff --git a/Engine/Build/Android/Java/src/com/epicgames/ue4/AudioDecoder.java b/Engine/Build/Android/Java/src/com/epicgames/ue4/AudioDecoder.java
--- a/Engine/Build/Android/Java/src/com/epicgames/ue4/AudioDecoder.java
+++ b/Engine/Build/Android/Java/src/com/epicgames/ue4/AudioDecoder.java
@@ -35,109 +35,131 @@ public class AudioDecoder
 	
 	public byte[] DecodeAudio(byte[] rawBuffer)
 	{
-		byte[] data = null;
-		if( mAudioCodec != null)
+		synchronized(this)
 		{
-			int inIndex = mAudioCodec.dequeueInputBuffer(-1);
-			if (inIndex >= 0) 
+			byte[] data = null;
+			if( mAudioCodec != null)
 			{
-				int sampleSize = rawBuffer.length;
-				ByteBuffer buffer = mAudioCodec.getInputBuffer(inIndex);
-				buffer.clear();
-				buffer.put(rawBuffer);
-				buffer.clear();
-
-				mAudioCodec.queueInputBuffer(inIndex, 0, sampleSize, 0, 0);
-			}
-
-			MediaCodec.BufferInfo info = new MediaCodec.BufferInfo();
-
-			int outIndex = mAudioCodec.dequeueOutputBuffer(info, 10000);
-
-			if (outIndex >= 0) 
-			{
-				ByteBuffer outBuffer = mAudioCodec.getOutputBuffer(outIndex);
-				data = new byte[info.size];
+				int inIndex = mAudioCodec.dequeueInputBuffer(-1);
+				if (inIndex >= 0) 
+				{
+					int sampleSize = rawBuffer.length;
+					ByteBuffer buffer = mAudioCodec.getInputBuffer(inIndex);
+					buffer.clear();
+					buffer.put(rawBuffer);
+					buffer.clear();
+
+					mAudioCodec.queueInputBuffer(inIndex, 0, sampleSize, 0, 0);
+				}
+
+				MediaCodec.BufferInfo info = new MediaCodec.BufferInfo();
+
+				int outIndex = mAudioCodec.dequeueOutputBuffer(info, 10000);
+
+				switch (outIndex) 
+				{
+					case MediaCodec.INFO_OUTPUT_BUFFERS_CHANGED:
+						GameActivity.Log.debug("Android Audio Decoder: INFO_OUTPUT_BUFFERS_CHANGED");
+						break;
+					case MediaCodec.INFO_OUTPUT_FORMAT_CHANGED:
+						GameActivity.Log.debug("Android Audio Decoder: New format " + mAudioCodec.getOutputFormat());
+						break;
+					case MediaCodec.INFO_TRY_AGAIN_LATER:
+						GameActivity.Log.debug("Android Audio Decoder: dequeueOutputBuffer timed out!");
+						break;
+					default:
+						ByteBuffer outBuffer = mAudioCodec.getOutputBuffer(outIndex);
+						data = new byte[info.size];
 			
-				outBuffer.get(data);
-				outBuffer.clear();
+						outBuffer.get(data);
+						outBuffer.clear();
 				
-				mAudioCodec.releaseOutputBuffer(outIndex, true);
+						mAudioCodec.releaseOutputBuffer(outIndex, true);
+				}
 			}
+			return data;
 		}
-		return data;
 	}
 
-	public void reset()
+	public void resetCodec()
 	{
-		release();
+		synchronized(this)
+		{
+			if (null != mAudioCodec)
+			{
+				try 
+				{
+					mAudioCodec.stop();
+					mAudioCodec.release();
+				}
+				catch(Exception e)
+				{
+					e.printStackTrace();
+				}
+			};
+		}
 	}
 
 	public void release()
 	{
-		if (null != mAudioCodec)
-		{
-			mAudioCodec.stop();
-			mAudioCodec.release();
-		}
+		resetCodec();
 	}
 
 	private boolean CreateCodec()
 	{
-		try 
+		synchronized(this)
 		{
-			// Create the decoder
-			mAudioCodec = MediaCodec.createDecoderByType(MediaFormat.MIMETYPE_AUDIO_AAC);
+			try 
+			{
+				// Create the decoder
+				mAudioCodec = MediaCodec.createDecoderByType(MediaFormat.MIMETYPE_AUDIO_AAC);
 
-			// Media settings
-			MediaFormat format = MediaFormat.createAudioFormat(MediaFormat.MIMETYPE_AUDIO_AAC, mSampleRate, mChannelCount);
+				// Media settings
+				MediaFormat format = MediaFormat.createAudioFormat(MediaFormat.MIMETYPE_AUDIO_AAC, mSampleRate, mChannelCount);
 
-			// PPS and SPS headers
-			format.setInteger(MediaFormat.KEY_BIT_RATE, mBitsPerSample);
+				// decoding params
+				format.setInteger(MediaFormat.KEY_BIT_RATE, mBitsPerSample);
 
-			byte[] bytes = new byte[]{(byte) 0x12, (byte)0x12};
-			ByteBuffer bb = ByteBuffer.wrap(bytes);
-			format.setByteBuffer("csd-0", bb);
+				// PPS and SPS headers
+				byte[] bytes = new byte[]{(byte) 0x12, (byte)0x12};
+				ByteBuffer bb = ByteBuffer.wrap(bytes);
+				format.setByteBuffer("csd-0", bb);
 			
-			// Configure the decoder
-			mAudioCodec.configure(format, null, null, 0);
+				// Configure the decoder
+				mAudioCodec.configure(format, null, null, 0);
 
-			// Start the decoder
-			mAudioCodec.start();
-			bIsInitialized = true;
-		}
-		catch(Exception e)
-		{
-			GameActivity.Log.warn("Android Audio Decoder: CreateCodec failed!");		
-			e.printStackTrace();
-			return false;
-		}
+				// Start the decoder
+				mAudioCodec.start();
+				bIsInitialized = true;
+			}
+			catch(Exception e)
+			{
+				GameActivity.Log.warn("Android Audio Decoder: CreateCodec failed!");		
+				e.printStackTrace();
+				return false;
+			}
 		
-		return true;
+			return true;
+		}
 	}
 
 	public boolean updateConfig(int sampleRate, int channelCount, int BitsPerSample, byte[] AacSpecificConfig)
 	{
-		bIsInitialized = false;
+		synchronized(this)
+		{
+			GameActivity.Log.warn("Android Audio Decoder: updateConfig channelCount:" + channelCount);		
+
+			bIsInitialized = false;
 		
-		mSampleRate = sampleRate;
-		mChannelCount = channelCount;
-		mBitsPerSample = BitsPerSample;
-		mAacSpecificConfig = AacSpecificConfig;
+			mSampleRate = sampleRate;
+			mChannelCount = channelCount;
+			mBitsPerSample = BitsPerSample;
+			mAacSpecificConfig = AacSpecificConfig;
 
-		if(mAudioCodec != null)
-		{
-			try 
-			{
-				mAudioCodec.stop();
-				mAudioCodec.release();
-			}
-			catch(Exception e)
-			{
-				e.printStackTrace();
-			}
+			resetCodec();
+
+			return CreateCodec();
 		}
-		return CreateCodec();
 	}
 
 }
diff --git a/Engine/Build/Android/Java/src/com/epicgames/ue4/VideoDecoder.java b/Engine/Build/Android/Java/src/com/epicgames/ue4/VideoDecoder.java
--- a/Engine/Build/Android/Java/src/com/epicgames/ue4/VideoDecoder.java
+++ b/Engine/Build/Android/Java/src/com/epicgames/ue4/VideoDecoder.java
@@ -200,9 +200,24 @@ public class VideoDecoder
 		return -1;
 	}
 
-	public void reset()
+	public void resetCodec()
 	{
-		release();
+		synchronized(this)
+		{
+			if (null != mVideoCodec)
+			{
+				try 
+				{
+					// Get an instance of MediaCodec and give it its Mime type
+					mVideoCodec.stop();
+					mVideoCodec.release();
+				}
+				catch(Exception e)
+				{
+					e.printStackTrace();
+				}
+			}
+		}
 	}
 
 	public void release()
@@ -217,11 +232,8 @@ public class VideoDecoder
 			while (WaitOnBitmapRender) ;
 			releaseOESTextureRenderer();
 		}
-		if (null != mVideoCodec)
-		{
-			mVideoCodec.stop();
-			mVideoCodec.release();
-		}
+
+		resetCodec();
 	}
 
 	private ByteBuffer getSPSPPSHeader()
@@ -305,28 +317,12 @@ public class VideoDecoder
 		SPS = inSPS;
 		PPS = inPPS;
 
-		release();
-
 		if(mVideoCodec == null)
 		{
-			if(SwizzlePixels || VulkanRenderer || !SupportsImageExternal)
-			{
-				//initBitmapRenderer();
-			}
-
 			return true;
 		}
 
-		try 
-		{
-			// Get an instance of MediaCodec and give it its Mime type
-			mVideoCodec.stop();
-			mVideoCodec.release();
-		}
-		catch(Exception e)
-		{
-			e.printStackTrace();
-		}
+		resetCodec();
 
 		return CreateCodec();
 	}
@@ -370,7 +366,7 @@ public class VideoDecoder
 			if (!CreateBitmapRenderer())
 			{
 				GameActivity.Log.warn("initBitmapRenderer failed to alloc mBitmapRenderer ");
-				reset();
+				release();
 			  }
 		}
 	}
@@ -1294,7 +1290,7 @@ public class VideoDecoder
 			if (!CreateOESTextureRenderer(externalTextureId))
 			{
 				GameActivity.Log.warn("updateVideoFrame failed to alloc mOESTextureRenderer ");
-				reset();
+				release();
 				return null;
 			}
 		}
