Add safezone handling for insets and display cutout

#ue4
#android

diff --git a/Engine/Build/Android/Java/src/com/epicgames/ue4/GameActivity.java.template b/Engine/Build/Android/Java/src/com/epicgames/ue4/GameActivity.java.template
--- a/Engine/Build/Android/Java/src/com/epicgames/ue4/GameActivity.java.template
+++ b/Engine/Build/Android/Java/src/com/epicgames/ue4/GameActivity.java.template
@@ -186,6 +186,7 @@ import com.epicgames.ue4.DownloadShim;
 import android.view.inputmethod.InputMethodManager;
 
 import android.graphics.Rect;
+import android.graphics.RectF;
 import android.view.ViewTreeObserver;
 
 import java.lang.reflect.Method;
@@ -229,7 +230,10 @@ public class GameActivity extends $${gameActivitySuperClass}$$ implements Surfac
 	private Sensor gyroscope;
 	static boolean first_acceleration_sample = true;
 	static final float SampleDecayRate = 0.85f;
-	
+
+	private Rect safezone = new Rect(0, 0, 0, 0);
+	private RectF safezoneF = new RectF(0.0f, 0.0f, 0.0f, 0.0f);
+
 	private final float[] rotationMatrix = new float[9];
 	private final float[] orientationAngles = new float[3];
 
@@ -2099,6 +2103,11 @@ public class GameActivity extends $${gameActivitySuperClass}$$ implements Surfac
 		return mainView;
 	}
 
+	protected Rect GetSafezone()
+	{
+		return safezone;
+	}
+
 	protected FrameLayout GetContainerFrameLayout()
 	{
 		return containerFrameLayout;
@@ -2808,6 +2817,49 @@ public class GameActivity extends $${gameActivitySuperClass}$$ implements Surfac
 
         mainDecorView = getWindow().getDecorView();
 		mainDecorViewRect = new Rect();
+		
+		// check for display cutouts to set safe zone
+		View mainRootView = mainDecorView.findViewById(android.R.id.content).getRootView();
+		mainRootView.setOnApplyWindowInsetsListener(new View.OnApplyWindowInsetsListener() {
+            @Override
+       	    public WindowInsets onApplyWindowInsets(View v, WindowInsets insets) {
+				if (Build.VERSION.SDK_INT >= 28)
+				{
+					DisplayCutout cutout = insets.getDisplayCutout();
+					if (cutout != null)
+					{
+						Rect fullViewRect = new Rect();
+						mainDecorView.getDrawingRect( mainDecorViewRect );
+
+						android.view.Display display = getWindowManager().getDefaultDisplay();
+						android.graphics.Point displaySize = new android.graphics.Point();
+						display.getRealSize(displaySize);
+
+						// store safezone independent of resolution
+						float invX = 1.0f / (displaySize.x != 0.0f ? displaySize.x : 1.0f);
+						float invY = 1.0f / (displaySize.y != 0.0f ? displaySize.y : 1.0f);
+						safezoneF.left = (safezone.left = cutout.getSafeInsetLeft()) * invX;
+						safezoneF.top = (safezone.top = cutout.getSafeInsetTop()) * invY;
+						safezoneF.right = (safezone.right = cutout.getSafeInsetRight()) * invX;
+						safezoneF.bottom = (safezone.bottom = cutout.getSafeInsetBottom()) *invY;
+						nativeSetSafezoneInfo(displaySize.y > displaySize.x, safezoneF.left, safezoneF.top, safezoneF.right, safezoneF.bottom);
+
+						//Log.debug("FullViewRect: " + fullViewRect.toString());	
+						//Log.debug("DisplayCutout: " + cutout.toString());
+						//Log.debug("SafeTop: " + cutout.getSafeInsetTop());
+						//Log.debug("SafeBottom: " + cutout.getSafeInsetBottom());
+						//Log.debug("SafeLeft: " + cutout.getSafeInsetLeft());
+						//Log.debug("SafeRight: " + cutout.getSafeInsetRight());
+						//List<Rect>bounds = cutout.getBoundingRects();
+						//for (Rect r : bounds)
+						//{
+						//	Log.debug("Rect: " + r.toString());
+						//}
+					}
+				}
+				return insets;
+			}
+        });
 
 		createVirtualKeyboardInput();
 		
@@ -5802,6 +5854,7 @@ public class GameActivity extends $${gameActivitySuperClass}$$ implements Surfac
 	public native void nativeSetAndroidVersionInformation( String AndroidVersion, String PhoneMake, String PhoneModel, String PhoneBuildNumber, String OSLanguage );
 
 	public native void nativeSetSurfaceViewInfo(int width, int height);
+	public native void nativeSetSafezoneInfo(boolean bIsPortrait, float left, float top, float right, float bottom);
 
 	public native void nativeConsoleCommand(String commandString);
 	public native void nativeVirtualKeyboardChanged(String contents);
diff --git a/Engine/Source/Runtime/ApplicationCore/Private/Android/AndroidApplication.cpp b/Engine/Source/Runtime/ApplicationCore/Private/Android/AndroidApplication.cpp
--- a/Engine/Source/Runtime/ApplicationCore/Private/Android/AndroidApplication.cpp
+++ b/Engine/Source/Runtime/ApplicationCore/Private/Android/AndroidApplication.cpp
@@ -188,17 +188,34 @@ void FDisplayMetrics::RebuildDisplayMetrics( FDisplayMetrics& OutDisplayMetrics
 	float Inset_Right = -1.0f;
 	float Inset_Bottom = -1.0f;
 
-	if (FString* SafeZoneLandscape = FAndroidMisc::GetConfigRulesVariable(TEXT("SafeZone_Landscape")))
+	bool bIsPortrait = FAndroidWindow::IsPortraitOrientation();
+
+	// ConfigRules values override values from device
+	if (FString* SafeZoneVar = FAndroidMisc::GetConfigRulesVariable(bIsPortrait ? TEXT("SafeZone_Portrait") : TEXT("SafeZone_Landscape")))
 	{
 		TArray<FString> ZoneVector;
-		if (SafeZoneLandscape->ParseIntoArray(ZoneVector, TEXT(","), true) == 4)
+		int ZoneParseCount = SafeZoneVar->ParseIntoArray(ZoneVector, TEXT(","), true);
+		ensureMsgf(ZoneParseCount == 4, TEXT("SafeZone variable not properly formatted."));
+
+		if (ZoneParseCount == 4)
 		{
+			// these are already in pixels
 			Inset_Left = FCString::Atof(*ZoneVector[0]);
 			Inset_Top = FCString::Atof(*ZoneVector[1]);
 			Inset_Right = FCString::Atof(*ZoneVector[2]);
 			Inset_Bottom = FCString::Atof(*ZoneVector[3]);
 		}
 	}
+	else
+	{
+		FVector4 SafeZoneRect = FAndroidWindow::GetSafezone(bIsPortrait);
+
+		// These values will be negative if there is not a safe zone set by GameActivity and need scaling
+		Inset_Left = SafeZoneRect.X * OutDisplayMetrics.PrimaryDisplayWidth;
+		Inset_Top = SafeZoneRect.Y * OutDisplayMetrics.PrimaryDisplayHeight;
+		Inset_Right = SafeZoneRect.Z * OutDisplayMetrics.PrimaryDisplayWidth;
+		Inset_Bottom = SafeZoneRect.W * OutDisplayMetrics.PrimaryDisplayHeight;
+	}
 
 	OutDisplayMetrics.TitleSafePaddingSize.X = (Inset_Left >= 0.0f) ? Inset_Left : OutDisplayMetrics.TitleSafePaddingSize.X;
 	OutDisplayMetrics.TitleSafePaddingSize.Y = (Inset_Top >= 0.0f) ? Inset_Top : OutDisplayMetrics.TitleSafePaddingSize.Y;
diff --git a/Engine/Source/Runtime/ApplicationCore/Private/Android/AndroidWindow.cpp b/Engine/Source/Runtime/ApplicationCore/Private/Android/AndroidWindow.cpp
--- a/Engine/Source/Runtime/ApplicationCore/Private/Android/AndroidWindow.cpp
+++ b/Engine/Source/Runtime/ApplicationCore/Private/Android/AndroidWindow.cpp
@@ -71,6 +71,8 @@ void FAndroidWindow::SetOSWindowHandle(void* InWindow)
 //This function is declared in the Java-defined class, GameActivity.java: "public native void nativeSetObbInfo(String PackageName, int Version, int PatchVersion);"
 static bool GAndroidIsPortrait = false;
 static int GAndroidDepthBufferPreference = 0;
+static FVector4 GAndroidPortraitSafezone = FVector4(-1.0f, -1.0f, -1.0f, -1.0f);
+static FVector4 GAndroidLandscapeSafezone = FVector4(-1.0f, -1.0f, -1.0f, -1.0f);
 #if USE_ANDROID_JNI
 JNI_METHOD void Java_com_epicgames_ue4_GameActivity_nativeSetWindowInfo(JNIEnv* jenv, jobject thiz, jboolean bIsPortrait, jint DepthBufferPreference)
 {
@@ -86,8 +88,36 @@ JNI_METHOD void Java_com_epicgames_ue4_GameActivity_nativeSetSurfaceViewInfo(JNI
 	GSurfaceViewHeight = height;
 	UE_LOG(LogAndroid, Log, TEXT("nativeSetSurfaceViewInfo width=%d and height=%d"), GSurfaceViewWidth, GSurfaceViewHeight);
 }
+
+JNI_METHOD void Java_com_epicgames_ue4_GameActivity_nativeSetSafezoneInfo(JNIEnv* jenv, jobject thiz, jboolean bIsPortrait, jfloat left, jfloat top, jfloat right, jfloat bottom)
+{
+	if (bIsPortrait)
+	{
+		GAndroidPortraitSafezone.X = left;
+		GAndroidPortraitSafezone.Y = top;
+		GAndroidPortraitSafezone.Z = right;
+		GAndroidPortraitSafezone.W = bottom;
+	}
+	else
+	{
+		GAndroidLandscapeSafezone.X = left;
+		GAndroidLandscapeSafezone.Y = top;
+		GAndroidLandscapeSafezone.Z = right;
+		GAndroidLandscapeSafezone.W = bottom;
+	}
+}
 #endif
 
+bool FAndroidWindow::IsPortraitOrientation()
+{
+	return GAndroidIsPortrait;
+}
+
+FVector4 FAndroidWindow::GetSafezone(bool bPortrait)
+{
+	return bPortrait ? GAndroidPortraitSafezone : GAndroidLandscapeSafezone;
+}
+
 int32 FAndroidWindow::GetDepthBufferPreference()
 {
 	return GAndroidDepthBufferPreference;
diff --git a/Engine/Source/Runtime/ApplicationCore/Public/Android/AndroidWindow.h b/Engine/Source/Runtime/ApplicationCore/Public/Android/AndroidWindow.h
--- a/Engine/Source/Runtime/ApplicationCore/Public/Android/AndroidWindow.h
+++ b/Engine/Source/Runtime/ApplicationCore/Public/Android/AndroidWindow.h
@@ -56,6 +56,9 @@ public:
 	 */
 	static void* WaitForHardwareWindow();
 
+	static bool IsPortraitOrientation();
+	static FVector4 GetSafezone(bool bPortrait);
+
 protected:
 	/** @return true if the native window is currently in fullscreen mode, false otherwise */
 	virtual EWindowMode::Type GetWindowMode() const override { return EWindowMode::Fullscreen; }
