Handle required permissions at startup in SplashActivity with better flexiblity

#ue4
#android

diff --git a/Engine/Build/Android/Java/res/values/permissions.xml b/Engine/Build/Android/Java/res/values/permissions.xml
new file mode 100644
--- /dev/null
+++ b/Engine/Build/Android/Java/res/values/permissions.xml
@@ -0,0 +1,21 @@
+<?xml version="1.0" encoding="utf-8"?>
+<resources>
+	<!-- Button text -->
+	<string name="PERM_Quit">Quit</string>
+	<string name="PERM_OK">OK</string>
+	<string name="PERM_Settings">Settings</string>
+
+	<!-- Dialog caption text -->
+	<string name="PERM_Caption_PermRequired">Permission Required</string>
+
+	<!-- Dialog message text -->
+	<string name="PERM_Info_ApproveSettings">You must approve this permission in App Settings:</string>
+	<string name="PERM_Info_android.permission.WRITE_EXTERNAL_STORAGE">Access to files is needed to read and write game data.</string>
+	<string name="PERM_Info_android.permission.READ_EXTERNAL_STORAGE">Access to files is needed to read and write game data.</string>
+	<string name="PERM_Info_android.permission.GET_ACCOUNTS">"Access to contacts is needed to verify account information.</string>
+
+	<!-- Name of permission in App Settings -->
+	<string name="PERM_SettingsName_android.permission.WRITE_EXTERNAL_STORAGE">Storage</string>
+	<string name="PERM_SettingsName_android.permission.READ_EXTERNAL_STORAGE">Storage</string>
+	<string name="PERM_SettingsName_android.permission.GET_ACCOUNTS">Contacts</string>
+</resources>
diff --git a/Engine/Build/Android/Java/src/com/epicgames/ue4/SplashActivity.java b/Engine/Build/Android/Java/src/com/epicgames/ue4/SplashActivity.java
--- a/Engine/Build/Android/Java/src/com/epicgames/ue4/SplashActivity.java
+++ b/Engine/Build/Android/Java/src/com/epicgames/ue4/SplashActivity.java
@@ -2,16 +2,25 @@
 
 package com.epicgames.ue4;
 
+import java.util.ArrayList;
+import java.util.List;
+
 import android.annotation.SuppressLint;
 import android.os.Bundle;
 import android.app.Activity;
+import android.app.AlertDialog;
+import android.content.Context;
+import android.content.DialogInterface;
 import android.content.Intent;
 import android.content.pm.ApplicationInfo;
 import android.content.pm.PackageInfo;
 import android.content.pm.PackageManager;
 import android.content.pm.PackageManager.NameNotFoundException;
+import android.content.pm.PermissionInfo;
+import android.provider.Settings;
 import android.view.View;
 import android.view.WindowManager;
+import android.net.Uri;
 
 import android.support.v4.app.ActivityCompat;
 import android.support.v4.content.ContextCompat;
@@ -19,11 +28,15 @@ import android.support.v4.content.ContextCompat;
 public class SplashActivity extends Activity
 {
 	private static final int PERMISSION_REQUEST_CODE = 1105;
+	private static final int REQUEST_PERMISSION_SETTING = 1;
+
+	private String packageName;
+	private PackageManager pm;
+	private String[] permissionsRequiredAtStart = {};
 
 	private Intent GameActivityIntent;
 	private boolean WaitForPermission = false;
 
-
 	public static Logger Log = new Logger("UE4-SplashActivity");
 	
 	@SuppressLint("ObsoleteSdkInt")
@@ -31,13 +44,15 @@ public class SplashActivity extends Activity
 	protected void onCreate(Bundle savedInstanceState)
 	{
 		super.onCreate(savedInstanceState);
+		packageName = getPackageName();
+		pm = getPackageManager();
 
 		boolean ShouldHideUI = false;
 		boolean UseDisplayCutout = false;
 		boolean IsShipping = false;
 		boolean UseExternalFilesDir = false;
 		try {
-			ApplicationInfo ai = getPackageManager().getApplicationInfo(getPackageName(), PackageManager.GET_META_DATA);
+			ApplicationInfo ai = pm.getApplicationInfo(packageName, PackageManager.GET_META_DATA);
 			Bundle bundle = ai.metaData;
 
 			if(bundle.containsKey("com.epicgames.ue4.GameActivity.bShouldHideUI"))
@@ -57,6 +72,10 @@ public class SplashActivity extends Activity
             {
                 UseExternalFilesDir = bundle.getBoolean("com.epicgames.ue4.GameActivity.bUseExternalFilesDir");
             }
+			if (bundle.containsKey("com.epicgames.ue4.GameActivity.StartupPermissions"))
+			{
+				permissionsRequiredAtStart = bundle.getString("com.epicgames.ue4.GameActivity.StartupPermissions").split(",");
+			}
 		}
 		catch (NameNotFoundException | NullPointerException e)
 		{
@@ -76,7 +95,6 @@ public class SplashActivity extends Activity
 											| View.SYSTEM_UI_FLAG_IMMERSIVE);  // NOT sticky (will be set later in MainActivity)
 			}
 		}
-
 		
 		// allow certain models for now to use full area around cutout
 		boolean BlockDisplayCutout = true;
@@ -142,36 +160,193 @@ public class SplashActivity extends Activity
 			GameActivityIntent.setAction(intentAction);
 		}
 
-		// check if we need to wait for permissions
-		int targetSdkVersion = 0;
-		try 
+		// make a list of ungranted dangerous permissions in manifest required at start of GameActivity and request any we still need
+		ArrayList<String> ungrantedPermissions = getUngrantedPermissions(this, getDangerousPermissions(pm, packageName), permissionsRequiredAtStart);
+		if (ungrantedPermissions.size() > 0)
 		{
-			PackageInfo packageInfo = getPackageManager().getPackageInfo(getPackageName(), 0);
-			targetSdkVersion = packageInfo.applicationInfo.targetSdkVersion;
+			ActivityCompat.requestPermissions(this, ungrantedPermissions.toArray(new String[ungrantedPermissions.size()]), PERMISSION_REQUEST_CODE);
+			WaitForPermission = true;
 		}
-		catch (PackageManager.NameNotFoundException e) 
+
+		if (!WaitForPermission)
 		{
+			startActivity(GameActivityIntent);
+			finish();
+			overridePendingTransition(0, 0);
 		}
+	}
 
-		if (android.os.Build.VERSION.SDK_INT >= 23 && targetSdkVersion >= 23) //23 is the API level (Marshmallow) where runtime permission handling is available
+	private int getResourceId(String VariableName, String ResourceName, String PackageName)
+	{
+		try {
+			return getResources().getIdentifier(VariableName, ResourceName, PackageName);
+		}
+		catch (Exception e) {
+			e.printStackTrace();
+			return -1;
+		} 
+	}
+
+	private String getResourceStringOrDefault(String PackageName, String ResourceName, String DefaultString)
+	{
+		int resourceId = getResourceId(ResourceName, "string", PackageName);
+		return (resourceId < 1) ? DefaultString : getString(resourceId);
+	}
+
+	public ArrayList<String> getUngrantedPermissions(Context context, ArrayList<String> dangerousPermissions, String[] requiredPermissions)
+	{
+		ArrayList<String> ungrantedPermissions = new ArrayList<>();
+		if (dangerousPermissions.size() > 0)
 		{
-			// we might need to ask for permission if we don't already have it
-			if (ContextCompat.checkSelfPermission(this, "android.permission.WRITE_EXTERNAL_STORAGE") != PackageManager.PERMISSION_GRANTED)
+			for (String required : requiredPermissions)
 			{
-				if (!IsShipping || !UseExternalFilesDir)
+				required = required.replaceAll("\\s", "");
+				if (required.length() > 0 && dangerousPermissions.contains(required))
 				{
-					ActivityCompat.requestPermissions(this, new String[] {"android.permission.WRITE_EXTERNAL_STORAGE"}, PERMISSION_REQUEST_CODE);
-					WaitForPermission = true;
+					if (ContextCompat.checkSelfPermission(context, required) != PackageManager.PERMISSION_GRANTED)
+					{
+						ungrantedPermissions.add(required);
+					}
 				}
 			}
 		}
+		return ungrantedPermissions;
+	}
 
-		if (!WaitForPermission)
+	public ArrayList<String> getDangerousPermissions(PackageManager pm, String packageName)
+	{
+		int targetSdkVersion = 0;
+		ArrayList<String> dangerousPermissions = new ArrayList<>();
+		try 
 		{
-			startActivity(GameActivityIntent);
-			finish();
-			overridePendingTransition(0, 0);
+			PackageInfo packageInfo = pm.getPackageInfo(packageName, 0);
+			targetSdkVersion = packageInfo.applicationInfo.targetSdkVersion;
+
+			// 23 is the API level (Marshmallow) where runtime permission handling is available
+			if (android.os.Build.VERSION.SDK_INT >= 23 && targetSdkVersion >= 23)
+			{
+				packageInfo = pm.getPackageInfo(packageName, PackageManager.GET_PERMISSIONS);
+				if (packageInfo != null)
+				{
+					if (packageInfo.requestedPermissions != null && packageInfo.requestedPermissions.length > 0)
+					{
+						if (android.os.Build.VERSION.SDK_INT >= 28)
+						{
+							for (String permission : packageInfo.requestedPermissions)
+							{
+								try
+								{
+									PermissionInfo permissionInfo = pm.getPermissionInfo(permission, 0);
+									if (permissionInfo.getProtection() == PermissionInfo.PROTECTION_DANGEROUS)
+									{
+										dangerousPermissions.add(permission);
+									}
+								}
+								catch (PackageManager.NameNotFoundException e)
+								{
+								}
+							}
+						}
+						else
+						{
+							for (String permission : packageInfo.requestedPermissions)
+							{
+								try
+								{
+									PermissionInfo permissionInfo = pm.getPermissionInfo(permission, 0);
+									if ((permissionInfo.protectionLevel & PermissionInfo.PROTECTION_MASK_BASE) == PermissionInfo.PROTECTION_DANGEROUS)
+									{
+										dangerousPermissions.add(permission);
+									}
+								}
+								catch (PackageManager.NameNotFoundException e)
+								{
+								}
+							}
+						}
+					}
+				}
+			}
+		}
+		catch (PackageManager.NameNotFoundException e) 
+		{
+		}
+
+		// if asking for WRITE_EXTERNAL_STORAGE, don't also need ask for READ_EXTERNAL_STORAGE
+		if (dangerousPermissions.contains("android.permission.WRITE_EXTERNAL_STORAGE"))
+		{
+			dangerousPermissions.remove("android.permission.READ_EXTERNAL_STORAGE");
 		}
+
+		return dangerousPermissions;
+	}
+
+	public String getRationale(String permission)
+	{
+		return getResourceStringOrDefault(packageName, "PERM_Info_" + permission, "This permission is required to start the game:\n" + permission);
+	}
+
+	public void showDialog(String title, String message, boolean bShowSettings)
+	{
+		final String dialogTitle = title;
+		final String dialogMessage = message;
+		final boolean dialogSettings = bShowSettings;
+
+		runOnUiThread(new Runnable()
+		{
+			@Override
+			public void run()
+			{
+				AlertDialog.Builder dialog = new AlertDialog.Builder(SplashActivity.this);
+				dialog.setCancelable(false);
+				dialog.setTitle(dialogTitle);
+				dialog.setMessage(dialogMessage);
+				dialog.setNegativeButton(getResourceStringOrDefault(packageName, "PERM_Quit", "Quit"), new DialogInterface.OnClickListener() {
+					@Override
+					public void onClick(DialogInterface dialog, int id) {
+						dialog.dismiss();
+						System.exit(0);
+					}
+				});
+				if (!dialogSettings)
+				{
+					dialog.setPositiveButton(getResourceStringOrDefault(packageName, "PERM_OK", "OK"), new DialogInterface.OnClickListener() {
+						@Override
+						public void onClick(DialogInterface dialog, int id) {
+							dialog.dismiss();
+							ArrayList<String> ungrantedPermissions = getUngrantedPermissions(SplashActivity.this, getDangerousPermissions(pm, packageName), permissionsRequiredAtStart);
+							if (ungrantedPermissions.size() > 0)
+							{
+								ActivityCompat.requestPermissions(SplashActivity.this, ungrantedPermissions.toArray(new String[ungrantedPermissions.size()]), PERMISSION_REQUEST_CODE);
+							}
+							else
+							{
+								// should not get here, but launch GameActivity since have all permissions
+								startActivity(GameActivityIntent);
+								finish();
+								overridePendingTransition(0, 0);
+							}
+						}
+					});
+				}
+				else
+				{
+					dialog.setPositiveButton(getResourceStringOrDefault(packageName, "PERM_Settings", "Settings"), new DialogInterface.OnClickListener() {
+						@Override
+						public void onClick(DialogInterface dialog, int id) {
+							dialog.dismiss();
+							Intent intent = new Intent(Settings.ACTION_APPLICATION_DETAILS_SETTINGS);
+							Uri uri = Uri.fromParts("package", packageName, null);
+							intent.setData(uri);
+							startActivityForResult(intent, REQUEST_PERMISSION_SETTING);
+							System.exit(0);
+						}
+					});
+				}
+
+				dialog.create().show();
+			}
+		});
 	}
 
 	@Override
@@ -179,16 +354,29 @@ public class SplashActivity extends Activity
 	{
 		if (requestCode==PERMISSION_REQUEST_CODE && permissions.length>0) 
 		{
-			if (grantResults.length>0 && grantResults[0]==PackageManager.PERMISSION_GRANTED)
+			for (int index=0; index < grantResults.length; index++)
 			{
-				startActivity(GameActivityIntent);
-				finish();
-				overridePendingTransition(0, 0);
-			}
-			else
-			{
-				finish();
+				String permission = permissions[index];
+				if (grantResults[index] == PackageManager.PERMISSION_DENIED)
+				{
+					boolean showRationale = ActivityCompat.shouldShowRequestPermissionRationale(SplashActivity.this, permission);
+					if (showRationale)
+					{
+						showDialog(getResourceStringOrDefault(packageName, "PERM_Caption_PermRequired", "Permissions Required"), getRationale(permission), false);
+					}
+					else
+					{
+						showDialog(getResourceStringOrDefault(packageName, "PERM_Caption_PermRequired", "Permissions Required"), getResourceStringOrDefault(packageName, "PERM_Info_ApproveSettings", "You must approve this permission in App Settings:") + "\n\n" +
+									getResourceStringOrDefault(packageName, "PERM_SettingsName_" + permission, permission), true);
+					}
+					return;
+				}
 			}
+
+			// all permissions granted, start GameActivity
+			startActivity(GameActivityIntent);
+			finish();
+			overridePendingTransition(0, 0);
 		}
 	}
 
@@ -203,4 +391,4 @@ public class SplashActivity extends Activity
 		}
 	}
 
-}
\ No newline at end of file
+}
diff --git a/Engine/Source/Programs/UnrealBuildTool/Platform/Android/UEDeployAndroid.cs b/Engine/Source/Programs/UnrealBuildTool/Platform/Android/UEDeployAndroid.cs
--- a/Engine/Source/Programs/UnrealBuildTool/Platform/Android/UEDeployAndroid.cs
+++ b/Engine/Source/Programs/UnrealBuildTool/Platform/Android/UEDeployAndroid.cs
@@ -2148,16 +2148,10 @@ namespace UnrealBuildTool
 			int StoreVersion = GetStoreVersion();
 
 			string Arch = GetNDKArch(UE4Arch);
-			int NDKLevelInt = ToolChain.GetNdkApiLevelInt();
-
-			// 64-bit targets must be android-21 or higher
-			if (NDKLevelInt < 21)
-			{
-				if (UE4Arch == "-arm64" || UE4Arch == "-x64")
-				{
-					NDKLevelInt = 21;
-				}
-			}
+			int NDKLevelInt = 0;
+			int MinSDKVersion = 0;
+			int TargetSDKVersion = 0;
+			GetMinTargetSDKVersions(ToolChain, UE4Arch, UPL, Arch, out MinSDKVersion, out TargetSDKVersion, out NDKLevelInt);
 
 			// get project version from ini
 			ConfigHierarchy GameIni = GetConfigCacheIni(ConfigHierarchyType.Game);
@@ -2174,10 +2168,6 @@ namespace UnrealBuildTool
 			Ini.GetBool("/Script/AndroidRuntimeSettings.AndroidRuntimeSettings", "bUseGetAccounts", out bUseGetAccounts);
 			string DepthBufferPreference;
 			Ini.GetString("/Script/AndroidRuntimeSettings.AndroidRuntimeSettings", "DepthBufferPreference", out DepthBufferPreference);
-			int MinSDKVersion;
-			Ini.GetInt32("/Script/AndroidRuntimeSettings.AndroidRuntimeSettings", "MinSDKVersion", out MinSDKVersion);
-			int TargetSDKVersion = MinSDKVersion;
-			Ini.GetInt32("/Script/AndroidRuntimeSettings.AndroidRuntimeSettings", "TargetSDKVersion", out TargetSDKVersion);
 			float MaxAspectRatioValue;
 			if (!Ini.TryGetValue("/Script/AndroidRuntimeSettings.AndroidRuntimeSettings", "MaxAspectRatio", out MaxAspectRatioValue))
 			{
@@ -2242,27 +2232,6 @@ namespace UnrealBuildTool
 					break;
 			}
 
-			// fix up the MinSdkVersion
-			if (NDKLevelInt > 19)
-			{
-				if (MinSDKVersion < 21)
-				{
-					MinSDKVersion = 21;
-					Log.TraceInformation("Fixing minSdkVersion; NDK level above 19 requires minSdkVersion of 21 (arch={0})", UE4Arch.Substring(1));
-				}
-			}
-
-			if (bGradleEnabled && MinSDKVersion < MinimumSDKLevelForGradle)
-			{
-				MinSDKVersion = MinimumSDKLevelForGradle;
-				Log.TraceInformation("Fixing minSdkVersion; requires minSdkVersion of {0} with Gradle based on active plugins", MinimumSDKLevelForGradle);
-			}
-
-			if (TargetSDKVersion < MinSDKVersion)
-			{
-				TargetSDKVersion = MinSDKVersion;
-			}
-
 			// only apply density to configChanges if using android-24 or higher and minimum sdk is 17
 			bool bAddDensity = (SDKLevelInt >= 24) && (MinSDKVersion >= 17);
 
@@ -2486,6 +2455,20 @@ namespace UnrealBuildTool
 				Text.AppendLine("\t\t<activity android:name=\".DownloaderActivity\" />");
 			}
 
+			// Figure out the required startup permissions if targetting devices supporting runtime permissions
+			String StartupPermissions = "";
+			if (TargetSDKVersion >= 23)
+			{
+				if (Configuration != "Shipping" || !bUseExternalFilesDir)
+				{
+					StartupPermissions = StartupPermissions + (StartupPermissions.Length > 0 ? "," : "") + "android.permission.WRITE_EXTERNAL_STORAGE";
+				}
+				if (bEnableGooglePlaySupport && bUseGetAccounts)
+				{
+					StartupPermissions = StartupPermissions + (StartupPermissions.Length > 0 ? "," : "") + "android.permission.GET_ACCOUNTS";
+				}
+			}
+
 			Text.AppendLine(string.Format("\t\t<meta-data android:name=\"com.epicgames.ue4.GameActivity.EngineVersion\" android:value=\"{0}\"/>", EngineVersion));
 			Text.AppendLine(string.Format("\t\t<meta-data android:name=\"com.epicgames.ue4.GameActivity.EngineBranch\" android:value=\"{0}\"/>", EngineBranch));
 			Text.AppendLine(string.Format("\t\t<meta-data android:name=\"com.epicgames.ue4.GameActivity.ProjectVersion\" android:value=\"{0}\"/>", ProjectVersion));
@@ -2503,6 +2486,7 @@ namespace UnrealBuildTool
 			Text.AppendLine(string.Format("\t\t<meta-data android:name=\"com.epicgames.ue4.GameActivity.bUseDisplayCutout\" android:value=\"{0}\"/>", bUseDisplayCutout ? "true" : "false"));
 			Text.AppendLine(string.Format("\t\t<meta-data android:name=\"com.epicgames.ue4.GameActivity.bAllowIMU\" android:value=\"{0}\"/>", bAllowIMU ? "true" : "false"));
 			Text.AppendLine(string.Format("\t\t<meta-data android:name=\"com.epicgames.ue4.GameActivity.bSupportsVulkan\" android:value=\"{0}\"/>", bSupportsVulkan ? "true" : "false"));
+			Text.AppendLine(string.Format("\t\t<meta-data android:name=\"com.epicgames.ue4.GameActivity.StartupPermissions\" android:value=\"{0}\"/>", StartupPermissions));
 			if (bUseNEONForArmV7)
 			{
 				Text.AppendLine("\t\t<meta-data android:name=\"com.epicgames.ue4.GameActivity.bUseNEONForArmV7\" android:value=\"{true}\"/>");
@@ -3037,26 +3021,34 @@ namespace UnrealBuildTool
 			return true;
 		}
 
-		private void GetMinTargetSDKVersions(AndroidToolChain ToolChain, string Arch, out int MinSDKVersion, out int TargetSDKVersion)
+		private void GetMinTargetSDKVersions(AndroidToolChain ToolChain, string Arch, UnrealPluginLanguage UPL, string NDKArch, out int MinSDKVersion, out int TargetSDKVersion, out int NDKLevelInt)
 		{
 			ConfigHierarchy Ini = GetConfigCacheIni(ConfigHierarchyType.Engine);
 			Ini.GetInt32("/Script/AndroidRuntimeSettings.AndroidRuntimeSettings", "MinSDKVersion", out MinSDKVersion);
 			TargetSDKVersion = MinSDKVersion;
 			Ini.GetInt32("/Script/AndroidRuntimeSettings.AndroidRuntimeSettings", "TargetSDKVersion", out TargetSDKVersion);
 
+			// Check for targetSDKOverride from UPL
+			string TargetOverride = UPL.ProcessPluginNode(NDKArch, "targetSDKOverride", "");
+			if (!String.IsNullOrEmpty(TargetOverride))
+			{
+				int OverrideInt = 0;
+				if (int.TryParse(TargetOverride, out OverrideInt))
+				{
+					TargetSDKVersion = OverrideInt;
+				}
+			}
+
 			// Make sure minSdkVersion is at least 13 (need this for appcompat-v13 used by AndroidPermissions)
 			// this may be changed by active plugins (Google Play Services 11.0.4 needs 14 for example)
-			if (MinSDKVersion < MinimumSDKLevelForGradle)
+			if (bGradleEnabled && MinSDKVersion < MinimumSDKLevelForGradle)
 			{
 				MinSDKVersion = MinimumSDKLevelForGradle;
-			}
-			if (TargetSDKVersion < MinSDKVersion)
-			{
-				TargetSDKVersion = MinSDKVersion;
+				Log.TraceInformation("Fixing minSdkVersion; requires minSdkVersion of {0} with Gradle based on active plugins", MinimumSDKLevelForGradle);
 			}
 
 			// 64-bit targets must be android-21 or higher
-			int NDKLevelInt = ToolChain.GetNdkApiLevelInt();
+			NDKLevelInt = ToolChain.GetNdkApiLevelInt();
 			if (NDKLevelInt < 21)
 			{
 				if (Arch == "-arm64" || Arch == "-x64")
@@ -3074,8 +3066,13 @@ namespace UnrealBuildTool
 					Log.TraceInformation("Fixing minSdkVersion; NDK level above 19 requires minSdkVersion of 21 (arch={0})", Arch.Substring(1));
 				}
 			}
+
+			if (TargetSDKVersion < MinSDKVersion)
+			{
+				TargetSDKVersion = MinSDKVersion;
+			}
 		}
-		
+
 		private void CreateGradlePropertiesFiles(string Arch, int MinSDKVersion, int TargetSDKVersion, string CompileSDKVersion, string BuildToolsVersion, string PackageName,
 			string DestApkName, string NDKArch,	string UE4BuildFilesPath, string GameBuildFilesPath, string UE4BuildGradleAppPath, string UE4BuildPath, string UE4BuildGradlePath, bool bForDistribution)
 		{
@@ -3963,7 +3960,8 @@ namespace UnrealBuildTool
 					// get min and target SDK versions
 					int MinSDKVersion = 0;
 					int TargetSDKVersion = 0;
-					GetMinTargetSDKVersions(ToolChain, Arch, out MinSDKVersion, out TargetSDKVersion);
+					int NDKLevelInt = 0;
+					GetMinTargetSDKVersions(ToolChain, Arch, UPL, NDKArch, out MinSDKVersion, out TargetSDKVersion, out NDKLevelInt);
 					
 					// move JavaLibs into subprojects
 					string JavaLibsDir = Path.Combine(UE4BuildPath, "JavaLibs");
