Added option to required permissions to only ask if in a list of specified device manufacturers

#android

diff --git a/Engine/Build/Android/Java/src/com/epicgames/ue4/SplashActivity.java b/Engine/Build/Android/Java/src/com/epicgames/ue4/SplashActivity.java
--- a/Engine/Build/Android/Java/src/com/epicgames/ue4/SplashActivity.java
+++ b/Engine/Build/Android/Java/src/com/epicgames/ue4/SplashActivity.java
@@ -74,7 +74,7 @@ public class SplashActivity extends Activity
             }
 			if (bundle.containsKey("com.epicgames.ue4.GameActivity.StartupPermissions"))
 			{
-				permissionsRequiredAtStart = bundle.getString("com.epicgames.ue4.GameActivity.StartupPermissions").split(",");
+				permissionsRequiredAtStart = filterRequiredPermissions(bundle.getString("com.epicgames.ue4.GameActivity.StartupPermissions"));
 			}
 		}
 		catch (NameNotFoundException | NullPointerException e)
@@ -193,6 +193,38 @@ public class SplashActivity extends Activity
 		return (resourceId < 1) ? DefaultString : getString(resourceId);
 	}
 
+	private String[] filterRequiredPermissions(String permissions)
+	{
+		String manufacturer = android.os.Build.MANUFACTURER;
+		ArrayList<String> keptPermissions = new ArrayList<>();
+		String[] requiredPermissions = permissions.split(",");
+		for (String required : requiredPermissions)
+		{
+			required = required.replaceAll("\\s", "");
+			if (required.length() > 0)
+			{
+				int conditionalIndex = required.indexOf("|");
+				if (conditionalIndex > 1)
+				{
+					String conditions[] = required.substring(1, conditionalIndex-1).split(",");
+					for (String make : conditions)
+					{
+						if (make.equals(manufacturer))
+						{
+							keptPermissions.add(required.substring(conditionalIndex + 1));
+							break;
+						}
+					}
+				}
+				else
+				{
+					keptPermissions.add(required);
+				}
+			}
+		}
+		return keptPermissions.toArray(new String[keptPermissions.size()]);
+	}
+
 	public ArrayList<String> getUngrantedPermissions(Context context, ArrayList<String> dangerousPermissions, String[] requiredPermissions)
 	{
 		ArrayList<String> ungrantedPermissions = new ArrayList<>();
@@ -200,8 +232,7 @@ public class SplashActivity extends Activity
 		{
 			for (String required : requiredPermissions)
 			{
-				required = required.replaceAll("\\s", "");
-				if (required.length() > 0 && dangerousPermissions.contains(required))
+				if (dangerousPermissions.contains(required))
 				{
 					if (ContextCompat.checkSelfPermission(context, required) != PackageManager.PERMISSION_GRANTED)
 					{
