Fix version checked for OBB to always match manifest version code

#android

diff --git a/Engine/Build/Android/Java/src/com/epicgames/ue4/GameActivity.java.template b/Engine/Build/Android/Java/src/com/epicgames/ue4/GameActivity.java.template
--- a/Engine/Build/Android/Java/src/com/epicgames/ue4/GameActivity.java.template
+++ b/Engine/Build/Android/Java/src/com/epicgames/ue4/GameActivity.java.template
@@ -223,6 +223,7 @@ public class GameActivity extends $${gameActivitySuperClass}$$ implements Surfac
 //$${gameActivityImplementsAdditions}$$
 															ComponentCallbacks2
 {
+	public int VersionCode = 0;
 	private boolean bAllowIMU = true;
 	private SensorManager sensorManager;
 	private Sensor accelerometer;
@@ -2605,9 +2606,9 @@ public class GameActivity extends $${gameActivitySuperClass}$$ implements Surfac
 
 		try
 		{
-			int Version = getPackageManager().getPackageInfo(getPackageName(), 0).versionCode;
+			VersionCode = getPackageManager().getPackageInfo(getPackageName(), 0).versionCode;
 			int PatchVersion = 0;
-			nativeSetObbInfo(ProjectName, getApplicationContext().getPackageName(), Version, PatchVersion, AppType);
+			nativeSetObbInfo(ProjectName, getApplicationContext().getPackageName(), VersionCode, PatchVersion, AppType);
 		}
 		catch (Exception e)
 		{
@@ -2652,6 +2653,7 @@ public class GameActivity extends $${gameActivitySuperClass}$$ implements Surfac
 						// remember down time and position
 						downTime = System.currentTimeMillis();
 						downX = event.getX();
+						consoleInputBox.requestFocus();
 						return true;
 					}
 					case MotionEvent.ACTION_UP: {
@@ -2795,7 +2797,7 @@ public class GameActivity extends $${gameActivitySuperClass}$$ implements Surfac
 		// check for OBB file present if we don't have all the files and don't need to verify
 		if (!HasAllFiles && !VerifyOBBOnStartUp)
 		{
-			HasAllFiles = DownloadShim.expansionFilesDelivered(this);
+			HasAllFiles = DownloadShim.expansionFilesDelivered(this, VersionCode);
 		}
 
 		containerFrameLayout = new FrameLayout(_activity);
diff --git a/Engine/Plugins/Runtime/GoogleVR/GoogleVRHMD/Source/GoogleVRHMD/GoogleVRHMD_APL.xml b/Engine/Plugins/Runtime/GoogleVR/GoogleVRHMD/Source/GoogleVRHMD/GoogleVRHMD_APL.xml
--- a/Engine/Plugins/Runtime/GoogleVR/GoogleVRHMD/Source/GoogleVRHMD/GoogleVRHMD_APL.xml
+++ b/Engine/Plugins/Runtime/GoogleVR/GoogleVRHMD/Source/GoogleVRHMD/GoogleVRHMD_APL.xml
@@ -557,7 +557,7 @@
 					{
 						if(!HasAllFiles)
 						{
-							HasAllFiles = DownloadShim.expansionFilesDelivered(this);
+							HasAllFiles = DownloadShim.expansionFilesDelivered(this, VersionCode);
 						}
 					}
 				</insert>
diff --git a/Engine/Source/Programs/UnrealBuildTool/Platform/Android/UEDeployAndroid.cs b/Engine/Source/Programs/UnrealBuildTool/Platform/Android/UEDeployAndroid.cs
--- a/Engine/Source/Programs/UnrealBuildTool/Platform/Android/UEDeployAndroid.cs
+++ b/Engine/Source/Programs/UnrealBuildTool/Platform/Android/UEDeployAndroid.cs
@@ -1037,9 +1037,9 @@ namespace UnrealBuildTool
 			ShimFileContent.Append("\tpublic static Class<DownloaderActivity> GetDownloaderType() { return DownloaderActivity.class; }\n");
 
 			// Do OBB file checking without using DownloadActivity to avoid transit to another activity
-			ShimFileContent.Append("\tpublic static boolean expansionFilesDelivered(Activity activity) {\n");
+			ShimFileContent.Append("\tpublic static boolean expansionFilesDelivered(Activity activity, int version) {\n");
 			ShimFileContent.Append("\t\tfor (OBBData.XAPKFile xf : OBBData.xAPKS) {\n");
-			ShimFileContent.Append("\t\t\tString fileName = Helpers.getExpansionAPKFileName(activity, xf.mIsMain, xf.mFileVersion, OBBData.AppType);\n");
+			ShimFileContent.Append("\t\t\tString fileName = Helpers.getExpansionAPKFileName(activity, xf.mIsMain, version, OBBData.AppType);\n");
 			ShimFileContent.Append("\t\t\tGameActivity.Log.debug(\"Checking for file : \" + fileName);\n");
 			ShimFileContent.Append("\t\t\tString fileForNewFile = Helpers.generateSaveFileName(activity, fileName);\n");
 			ShimFileContent.Append("\t\t\tString fileForDevFile = Helpers.generateSaveFileNameDevelopment(activity, fileName);\n");
