Fixed issue with streaming video timestamps on Android causing the stream not to play.

diff --git a/Engine/Build/Android/Java/src/com/epicgames/ue4/VideoDecoder.java b/Engine/Build/Android/Java/src/com/epicgames/ue4/VideoDecoder.java
--- a/Engine/Build/Android/Java/src/com/epicgames/ue4/VideoDecoder.java
+++ b/Engine/Build/Android/Java/src/com/epicgames/ue4/VideoDecoder.java
@@ -152,10 +152,24 @@ public class VideoDecoder
 					MediaCodec.BufferInfo info = new MediaCodec.BufferInfo();
 
 					int outIndex = mVideoCodec.dequeueOutputBuffer(info, -1);
-					if (outIndex >= 0) 
+
+					// manage the config frames
+					switch (outIndex) 
 					{
-						addFrameTimeDuration(timeStamp, duration); //us to ms
-						mVideoCodec.releaseOutputBuffer(outIndex, true);
+						case MediaCodec.INFO_OUTPUT_BUFFERS_CHANGED:
+							GameActivity.Log.debug("Android Video Decoder: INFO_OUTPUT_BUFFERS_CHANGED");
+							break;
+						case MediaCodec.INFO_OUTPUT_FORMAT_CHANGED:
+							GameActivity.Log.debug("Android Video Decoder: New format" + mVideoCodec.getOutputFormat());
+							break;
+						case MediaCodec.INFO_TRY_AGAIN_LATER:
+							GameActivity.Log.debug("Android Video Decoder: dequeueOutputBuffer timed out!");
+							break;
+						default:
+							addFrameTimeDuration(timeStamp, duration); //us to ms
+							
+							boolean doRender = (info.size != 0) && ((info.flags & MediaCodec.BUFFER_FLAG_CODEC_CONFIG) != MediaCodec.BUFFER_FLAG_CODEC_CONFIG);
+							mVideoCodec.releaseOutputBuffer(outIndex, doRender);
 					}
 
 					if ((info.flags & MediaCodec.BUFFER_FLAG_END_OF_STREAM) != 0) 
