Allow -epicapp or -androidprod to force production from commandline

#android

diff --git a/Engine/Build/Android/Java/src/com/epicgames/ue4/GameActivity.java.template b/Engine/Build/Android/Java/src/com/epicgames/ue4/GameActivity.java.template
--- a/Engine/Build/Android/Java/src/com/epicgames/ue4/GameActivity.java.template
+++ b/Engine/Build/Android/Java/src/com/epicgames/ue4/GameActivity.java.template
@@ -313,9 +313,10 @@ public class GameActivity extends $${gameActivitySuperClass}$$ implements Surfac
 	protected ProgressDialog mProgressDialog;
 
 	// Commandline arguments and flags
-	String CommandLineFull = "";
-	ArrayList<String> CommandLineArguments = new ArrayList<String>();
-	HashMap<String, String> CommandLineFlags = new HashMap<String, String>();
+	private boolean bHaveParsedCommandline = false;
+	public String CommandLineFull = "";
+	public ArrayList<String> CommandLineArguments = new ArrayList<String>();
+	public HashMap<String, String> CommandLineFlags = new HashMap<String, String>();
 
 	// Engine branch from manifest
 	String EngineBranch = "UE4";
@@ -628,8 +629,14 @@ public class GameActivity extends $${gameActivitySuperClass}$$ implements Surfac
         return packages.contains(packagename);
     }
 
-	private void ParseCommandline(String ProjectName)
+	public void ParseCommandline(String ProjectName)
 	{
+		if (bHaveParsedCommandline)
+		{
+			return;
+		}
+		bHaveParsedCommandline = true;
+
 		File file;
 		InputStream stream = null;
 		BufferedReader reader = null;
@@ -2117,7 +2124,9 @@ public class GameActivity extends $${gameActivitySuperClass}$$ implements Surfac
 	@Override
 	public void onCreate(Bundle savedInstanceState)
 	{
+		_activity = this;
 		super.onCreate(savedInstanceState);
+
 		sensorManager = (SensorManager) getSystemService(Context.SENSOR_SERVICE);
 		accelerometer = sensorManager.getDefaultSensor(Sensor.TYPE_ACCELEROMETER);
 		magnetometer = sensorManager.getDefaultSensor(Sensor.TYPE_MAGNETIC_FIELD);
@@ -2215,8 +2224,6 @@ public class GameActivity extends $${gameActivitySuperClass}$$ implements Surfac
 			PermissionHelper.acquirePermissions(new String[] {"android.permission.WRITE_EXTERNAL_STORAGE"}, this);
 		}
 
-		_activity = this;
-
 		// update memory stats every 10 seconds
 		memoryRunnable = new Runnable() {
 			@Override
