Enable fullscreen for Nokia 8.1

#android

diff --git a/Engine/Build/Android/Java/src/com/epicgames/ue4/SplashActivity.java b/Engine/Build/Android/Java/src/com/epicgames/ue4/SplashActivity.java
--- a/Engine/Build/Android/Java/src/com/epicgames/ue4/SplashActivity.java
+++ b/Engine/Build/Android/Java/src/com/epicgames/ue4/SplashActivity.java
@@ -102,6 +102,14 @@ public class SplashActivity extends Activity
 		{
 			BlockDisplayCutout = false;
 		}
+		else if (android.os.Build.MANUFACTURER.equals("HMD Global"))
+		{
+			String model = android.os.Build.MODEL;
+			if (model.equals("Nokia 8.1"))
+			{
+				BlockDisplayCutout = false;
+			}
+		}
 		else if (android.os.Build.MANUFACTURER.equals("samsung"))
 		{
 			String model = android.os.Build.MODEL;
