Add option for public logs on Android

#ue4
#android

diff --git a/Engine/Build/Android/Java/src/com/epicgames/ue4/GameActivity.java.template b/Engine/Build/Android/Java/src/com/epicgames/ue4/GameActivity.java.template
--- a/Engine/Build/Android/Java/src/com/epicgames/ue4/GameActivity.java.template
+++ b/Engine/Build/Android/Java/src/com/epicgames/ue4/GameActivity.java.template
@@ -404,6 +404,9 @@ public class GameActivity extends $${gameActivitySuperClass}$$ implements Surfac
 	/** Use ExternalFilesDir for UE4Game files */
     private boolean UseExternalFilesDir = false;
 
+	/** Make log files public */
+    private boolean PublicLogFiles = false;
+
 	/** Flag to ensure we have finished startup before allowing nativeOnActivityResult to get called */
 	private boolean InitCompletedOK = false;
 
@@ -2543,6 +2546,16 @@ public class GameActivity extends $${gameActivitySuperClass}$$ implements Surfac
                 Log.debug( "bUseExternalFilesDir not found. Leaving as " + UseExternalFilesDir);
             }
 
+			if (bundle.containsKey("com.epicgames.ue4.GameActivity.bPublicLogFiles"))
+            {
+                PublicLogFiles = bundle.getBoolean("com.epicgames.ue4.GameActivity.bPublicLogFiles");
+                Log.debug( "PublicLogFiles set to " + PublicLogFiles);
+            }
+            else
+            {
+                Log.debug( "PublicLogFiles not found. Leaving as " + PublicLogFiles);
+            }
+
 			if(bundle.containsKey("com.epicgames.ue4.GameActivity.bAllowIMU"))
 			{
 				bAllowIMU = bundle.getBoolean("com.epicgames.ue4.GameActivity.bAllowIMU");
@@ -2590,7 +2603,7 @@ public class GameActivity extends $${gameActivitySuperClass}$$ implements Surfac
 
 		Log.debug("APK path: " + getPackageResourcePath());
 		Log.debug("OBB in APK: " + (PackageDataInsideApkValue==1));
-		nativeSetGlobalActivity(UseExternalFilesDir, getFilesDir().getPath(), getExternalFilesDir(null).getPath(), PackageDataInsideApkValue==1, getPackageResourcePath());
+		nativeSetGlobalActivity(UseExternalFilesDir, PublicLogFiles, getFilesDir().getPath(), getExternalFilesDir(null).getPath(), PackageDataInsideApkValue==1, getPackageResourcePath());
 
 		// tell the engine if this is a portrait app
 		nativeSetWindowInfo(getResources().getConfiguration().orientation == Configuration.ORIENTATION_PORTRAIT, DepthBufferPreference);
@@ -5889,7 +5902,7 @@ public class GameActivity extends $${gameActivitySuperClass}$$ implements Surfac
 
 	public native boolean nativeIsShippingBuild();
 	public native void nativeSetAndroidStartupState(boolean bDebuggerAttached);
-	public native void nativeSetGlobalActivity(boolean bUseExternalFilesDir, String internalFilePath, String externalFilePath, boolean bOBBInAPK, String APKPath);
+	public native void nativeSetGlobalActivity(boolean bUseExternalFilesDir, boolean bPublicLogFiles, String internalFilePath, String externalFilePath, boolean bOBBInAPK, String APKPath);
 	public native void nativeSetWindowInfo(boolean bIsPortrait, int DepthBufferPreference);
 	public native void nativeSetObbInfo(String ProjectName, String PackageName, int Version, int PatchVersion, String AppType);
 	public native void nativeSetAndroidVersionInformation( String AndroidVersion, String PhoneMake, String PhoneModel, String PhoneBuildNumber, String OSLanguage );
diff --git a/Engine/Config/BaseEngine.ini b/Engine/Config/BaseEngine.ini
--- a/Engine/Config/BaseEngine.ini
+++ b/Engine/Config/BaseEngine.ini
@@ -1942,6 +1942,7 @@ bDetectVulkanByDefault=true
 bSplitIntoSeparateApks=false
 bPackageDataInsideApk=false
 bUseExternalFilesDir=false
+bPublicLogFiles=true
 bCreateAllPlatformsInstall=false
 Orientation=Landscape
 InstallLocation=InternalOnly
diff --git a/Engine/Source/Programs/UnrealBuildTool/Platform/Android/UEDeployAndroid.cs b/Engine/Source/Programs/UnrealBuildTool/Platform/Android/UEDeployAndroid.cs
--- a/Engine/Source/Programs/UnrealBuildTool/Platform/Android/UEDeployAndroid.cs
+++ b/Engine/Source/Programs/UnrealBuildTool/Platform/Android/UEDeployAndroid.cs
@@ -2249,6 +2249,13 @@ namespace UnrealBuildTool
 				Log.TraceInformation("Daydream and IMU both enabled, recommend disabling IMU if not needed.");
 			}
 
+			bool bPublicLogFiles = true;
+			Ini.GetBool("/Script/AndroidRuntimeSettings.AndroidRuntimeSettings", "bPublicLogFiles", out bPublicLogFiles);
+			if (!bUseExternalFilesDir)
+			{
+				bPublicLogFiles = false;
+			}
+
 			string InstallLocation;
 			Ini.GetString("/Script/AndroidRuntimeSettings.AndroidRuntimeSettings", "InstallLocation", out InstallLocation);
 			switch (InstallLocation.ToLower())
@@ -2515,6 +2522,7 @@ namespace UnrealBuildTool
 			Text.AppendLine(string.Format("\t\t<meta-data android:name=\"com.epicgames.ue4.GameActivity.CookedFlavors\" android:value=\"{0}\"/>", CookedFlavors));
 			Text.AppendLine(string.Format("\t\t<meta-data android:name=\"com.epicgames.ue4.GameActivity.bValidateTextureFormats\" android:value=\"{0}\"/>", bValidateTextureFormats ? "true" : "false"));
 			Text.AppendLine(string.Format("\t\t<meta-data android:name=\"com.epicgames.ue4.GameActivity.bUseExternalFilesDir\" android:value=\"{0}\"/>", bUseExternalFilesDir ? "true" : "false"));
+			Text.AppendLine(string.Format("\t\t<meta-data android:name=\"com.epicgames.ue4.GameActivity.bPublicLogFiles\" android:value=\"{0}\"/>", bPublicLogFiles ? "true" : "false"));
 			Text.AppendLine(string.Format("\t\t<meta-data android:name=\"com.epicgames.ue4.GameActivity.bUseDisplayCutout\" android:value=\"{0}\"/>", bUseDisplayCutout ? "true" : "false"));
 			Text.AppendLine(string.Format("\t\t<meta-data android:name=\"com.epicgames.ue4.GameActivity.bAllowIMU\" android:value=\"{0}\"/>", bAllowIMU ? "true" : "false"));
 			Text.AppendLine(string.Format("\t\t<meta-data android:name=\"com.epicgames.ue4.GameActivity.bSupportsVulkan\" android:value=\"{0}\"/>", bSupportsVulkan ? "true" : "false"));
diff --git a/Engine/Source/Runtime/Android/AndroidRuntimeSettings/Classes/AndroidRuntimeSettings.h b/Engine/Source/Runtime/Android/AndroidRuntimeSettings/Classes/AndroidRuntimeSettings.h
--- a/Engine/Source/Runtime/Android/AndroidRuntimeSettings/Classes/AndroidRuntimeSettings.h
+++ b/Engine/Source/Runtime/Android/AndroidRuntimeSettings/Classes/AndroidRuntimeSettings.h
@@ -254,6 +254,11 @@ public:
 	UPROPERTY(GlobalConfig, EditAnywhere, Category = "APK Packaging", Meta = (DisplayName = "Use ExternalFilesDir for UE4Game files?"))
 	bool bUseExternalFilesDir;
 
+	// If checked, log files will always be placed in a publicly available directory (either /sdcard/Android or /sdcard/UE4Game).
+	// You may require WRITE_EXTERNAL_STORAGE permission if you do not use ExternalFilesDir checkbox in android api 23+
+	UPROPERTY(GlobalConfig, EditAnywhere, Category = "APK Packaging", Meta = (DisplayName = "Make log files always publicly accessible?"))
+	bool bPublicLogFiles;
+
 	// The permitted orientation of the application on the device
 	UPROPERTY(GlobalConfig, EditAnywhere, Category = "APK Packaging")
 	TEnumAsByte<EAndroidScreenOrientation::Type> Orientation;
diff --git a/Engine/Source/Runtime/Core/Private/Android/AndroidPlatformFile.cpp b/Engine/Source/Runtime/Core/Private/Android/AndroidPlatformFile.cpp
--- a/Engine/Source/Runtime/Core/Private/Android/AndroidPlatformFile.cpp
+++ b/Engine/Source/Runtime/Core/Private/Android/AndroidPlatformFile.cpp
@@ -92,6 +92,10 @@ FString GFontPathBase;
 bool GOBBinAPK;
 FString GAPKFilename;
 
+// Directory for log file on Android
+bool GOverrideAndroidLogDir = false;
+static FString AndroidLogDir;
+
 #define FILEBASE_DIRECTORY "/UE4Game/"
 
 extern jobject AndroidJNI_GetJavaAssetManager();
@@ -932,6 +936,11 @@ public:
 	using IAndroidPlatformFile::IterateDirectory;
 	using IAndroidPlatformFile::IterateDirectoryStat;
 
+	static const FString* GetOverrideLogDirectory()
+	{
+		return GOverrideAndroidLogDir ? &AndroidLogDir : nullptr;
+	}
+
 	// On initialization we search for OBBs that we need to
 	// open to find resources.
 	virtual bool Initialize(IPlatformFile* Inner, const TCHAR* CmdLine) override
@@ -1004,6 +1013,19 @@ public:
 		mkdir(TCHAR_TO_UTF8(*FileBaseDir), 0766);
 		mkdir(TCHAR_TO_UTF8(*(FileBaseDir + GAndroidProjectName)), 0766);
 
+		// make sure the log directory exists if override applied
+		if (GOverrideAndroidLogDir)
+		{
+			FString LogBaseDir = GExternalFilePath + FString(FILEBASE_DIRECTORY);
+			mkdir(TCHAR_TO_UTF8(*LogBaseDir), 0766);
+			mkdir(TCHAR_TO_UTF8(*(LogBaseDir + GAndroidProjectName)), 0766);
+			mkdir(TCHAR_TO_UTF8(*(LogBaseDir + GAndroidProjectName + TEXT("/") + GAndroidProjectName)), 0766);
+			mkdir(TCHAR_TO_UTF8(*(LogBaseDir + GAndroidProjectName + TEXT("/") + GAndroidProjectName + TEXT("/Saved"))), 0766);
+			mkdir(TCHAR_TO_UTF8(*(LogBaseDir + GAndroidProjectName + TEXT("/") + GAndroidProjectName + TEXT("/Saved/Logs"))), 0766);
+
+			AndroidLogDir = LogBaseDir + GAndroidProjectName + TEXT("/") + GAndroidProjectName + TEXT("/Saved/Logs/");
+		}
+
 		return true;
 	}
 
@@ -2017,4 +2039,9 @@ IAndroidPlatformFile & IAndroidPlatformFile::GetPlatformPhysical()
 {
 	return FAndroidPlatformFile::GetPlatformPhysical();
 }
+
+const FString* IAndroidPlatformFile::GetOverrideLogDirectory()
+{
+	return FAndroidPlatformFile::GetOverrideLogDirectory();
+}
 #endif
diff --git a/Engine/Source/Runtime/Core/Private/Misc/Paths.cpp b/Engine/Source/Runtime/Core/Private/Misc/Paths.cpp
--- a/Engine/Source/Runtime/Core/Private/Misc/Paths.cpp
+++ b/Engine/Source/Runtime/Core/Private/Misc/Paths.cpp
@@ -377,6 +377,12 @@ FString FPaths::ProjectLogDir()
 	{
 		return FPlatformProcess::UserLogsDir();
 	}
+#elif PLATFORM_ANDROID && USE_ANDROID_FILE
+	const FString* OverrideDir = IAndroidPlatformFile::GetOverrideLogDirectory();
+	if (OverrideDir != nullptr)
+	{
+		return *OverrideDir;
+	}
 #endif
 
 	return FPaths::ProjectSavedDir() + TEXT("Logs/");
diff --git a/Engine/Source/Runtime/Core/Public/Android/AndroidPlatformFile.h b/Engine/Source/Runtime/Core/Public/Android/AndroidPlatformFile.h
--- a/Engine/Source/Runtime/Core/Public/Android/AndroidPlatformFile.h
+++ b/Engine/Source/Runtime/Core/Public/Android/AndroidPlatformFile.h
@@ -20,6 +20,14 @@ class CORE_API IAndroidPlatformFile : public IPhysicalPlatformFile
 public:
 	static IAndroidPlatformFile & GetPlatformPhysical();
 
+#if USE_ANDROID_FILE
+	/**
+	 * Get the directory path to write log files to.
+	 * This is /temp0 in shipping, or a path inside /data for other configs.
+	 */
+	static const FString* GetOverrideLogDirectory();
+#endif
+
 #if USE_ANDROID_JNI
 	// Get the android.content.res.AssetManager that Java code
 	// should use to open APK assets.
diff --git a/Engine/Source/Runtime/Launch/Private/Android/AndroidJNI.cpp b/Engine/Source/Runtime/Launch/Private/Android/AndroidJNI.cpp
--- a/Engine/Source/Runtime/Launch/Private/Android/AndroidJNI.cpp
+++ b/Engine/Source/Runtime/Launch/Private/Android/AndroidJNI.cpp
@@ -41,6 +41,7 @@ extern FString GInternalFilePath;
 extern FString GExternalFilePath;
 extern FString GFontPathBase;
 extern bool GOBBinAPK;
+extern bool GOverrideAndroidLogDir;
 extern FString GOBBFilePathBase;
 extern FString GAPKFilename;
 
@@ -1575,7 +1576,7 @@ JNIEXPORT jint JNI_OnLoad(JavaVM* InJavaVM, void* InReserved)
 //Native-defined functions
 
 //This function is declared in the Java-defined class, GameActivity.java: "public native void nativeSetGlobalActivity();"
-JNI_METHOD void Java_com_epicgames_ue4_GameActivity_nativeSetGlobalActivity(JNIEnv* jenv, jobject thiz, jboolean bUseExternalFilesDir, jstring internalFilePath, jstring externalFilePath, jboolean bOBBinAPK, jstring APKFilename /*, jobject googleServices*/)
+JNI_METHOD void Java_com_epicgames_ue4_GameActivity_nativeSetGlobalActivity(JNIEnv* jenv, jobject thiz, jboolean bUseExternalFilesDir, jboolean bPublicLogFiles, jstring internalFilePath, jstring externalFilePath, jboolean bOBBinAPK, jstring APKFilename /*, jobject googleServices*/)
 {
 	if (!FJavaWrapper::GameActivityThis)
 	{
@@ -1601,11 +1602,12 @@ JNI_METHOD void Java_com_epicgames_ue4_GameActivity_nativeSetGlobalActivity(JNIE
 		GAPKFilename = FJavaHelper::FStringFromParam(jenv, APKFilename);
 		GInternalFilePath = FJavaHelper::FStringFromParam(jenv, internalFilePath);
 		GExternalFilePath = FJavaHelper::FStringFromParam(jenv, externalFilePath);
-		
+
 		if (bUseExternalFilesDir)
 		{
 #if UE_BUILD_SHIPPING
 			GFilePathBase = GInternalFilePath;
+			GOverrideAndroidLogDir = bPublicLogFiles;
 #else
 			GFilePathBase = GExternalFilePath;
 #endif
