Fixed notifications on Android to be able to span on multiple lines.
#rnx

diff --git a/Engine/Build/Android/Java/src/com/epicgames/ue4/LocalNotificationReceiver.java b/Engine/Build/Android/Java/src/com/epicgames/ue4/LocalNotificationReceiver.java
--- a/Engine/Build/Android/Java/src/com/epicgames/ue4/LocalNotificationReceiver.java
+++ b/Engine/Build/Android/Java/src/com/epicgames/ue4/LocalNotificationReceiver.java
@@ -55,7 +55,8 @@ public class LocalNotificationReceiver extends BroadcastReceiver
 			.setContentIntent(pendingNotificationIntent)
 			.setWhen(System.currentTimeMillis())
 			.setTicker(details)		// note: will not show up on Lollipop up except for accessibility
-			.setContentTitle(title);
+			.setContentTitle(title)
+			.setStyle(new NotificationCompat.BigTextStyle().bigText(details));			
 		if (android.os.Build.VERSION.SDK_INT >= 21)
 		{
 			builder.setContentText(details);
