Scheduled local notifications on Android will now persist across reboots like on iOS.
#rnx

diff --git a/Engine/Build/Android/Java/src/com/epicgames/ue4/BootCompleteReceiver.java b/Engine/Build/Android/Java/src/com/epicgames/ue4/BootCompleteReceiver.java
new file mode 100644
--- /dev/null
+++ b/Engine/Build/Android/Java/src/com/epicgames/ue4/BootCompleteReceiver.java
@@ -0,0 +1,64 @@
+// Copyright 2019 Epic Games, Inc. All Rights Reserved.
+
+package com.epicgames.ue4;
+
+import android.content.BroadcastReceiver;
+import android.content.Context;
+import android.content.Intent;
+import android.content.SharedPreferences;
+
+import org.json.JSONException;
+import org.json.JSONObject;
+
+import java.util.Iterator;
+
+public class BootCompleteReceiver extends BroadcastReceiver
+{
+	@Override
+	public void onReceive(Context context, Intent intent)
+	{
+		// restore any scheduled notifications
+		SharedPreferences preferences = context.getSharedPreferences("LocalNotificationPreferences", Context.MODE_PRIVATE);
+		SharedPreferences.Editor editor = preferences.edit();
+		try
+		{
+			boolean changed = false;
+			JSONObject notificationDetails = new JSONObject(preferences.getString("notificationDetails", "{}"));
+			for (Iterator<String> i = notificationDetails.keys(); i.hasNext(); )
+			{
+				try
+				{
+					String key = i.next();
+					int notificationId = Integer.parseInt(key); 
+					JSONObject details = notificationDetails.getJSONObject(key);
+					String targetDateTime = details.getString("local-notification-targetDateTime");
+					boolean localTime = details.getBoolean("local-notification-localTime");
+					String title = details.getString("local-notification-title");
+					String body = details.getString("local-notification-body");
+					String action = details.getString("local-notification-action");
+					String activationEvent = details.getString("local-notification-activationEvent");
+					
+					if (!GameActivity.LocalNotificationScheduleAtTime(context, notificationId, targetDateTime, localTime, title, body, action, activationEvent))
+					{
+						// if it fails, remove from details
+						i.remove();
+					}
+				}
+				catch (NumberFormatException | JSONException e)
+				{
+					e.printStackTrace();
+					i.remove();
+				}
+			}
+			
+			if (changed)
+			{
+				editor.commit();
+			}
+		}
+		catch (JSONException e)
+		{
+			e.printStackTrace();
+		}
+	}
+}
diff --git a/Engine/Build/Android/Java/src/com/epicgames/ue4/GameActivity.java.template b/Engine/Build/Android/Java/src/com/epicgames/ue4/GameActivity.java.template
--- a/Engine/Build/Android/Java/src/com/epicgames/ue4/GameActivity.java.template
+++ b/Engine/Build/Android/Java/src/com/epicgames/ue4/GameActivity.java.template
@@ -4589,7 +4589,7 @@ public class GameActivity extends $${gameActivitySuperClass}$$ implements Surfac
 				localNotificationLaunchActivationEvent = extrasBundle.get("localNotificationLaunchActivationEvent").toString();
 				int notificationID = extrasBundle.getInt("localNotificationID");
 
-				LocalNotificationRemoveID(notificationID);
+				LocalNotificationRemoveID(this, notificationID);
 
 				// TODO
 				localNotificationLaunchFireDate = 0; 
@@ -4678,9 +4678,9 @@ public class GameActivity extends $${gameActivitySuperClass}$$ implements Surfac
 		return false;
 	}
 
-	private void LocalNotificationRemoveID(int notificationID)
+	public static void LocalNotificationRemoveID(Context context, int notificationID)
 	{
-		SharedPreferences preferences = getApplicationContext().getSharedPreferences("LocalNotificationPreferences", MODE_PRIVATE);
+		SharedPreferences preferences = context.getSharedPreferences("LocalNotificationPreferences", MODE_PRIVATE);
 		SharedPreferences.Editor editor = preferences.edit();
 		String notificationIDs = preferences.getString("notificationIDs", "");
 
@@ -4719,13 +4719,72 @@ public class GameActivity extends $${gameActivitySuperClass}$$ implements Surfac
 		editor.putString("notificationIDs", newNotificationIDs);
 		editor.commit();
 	}
+	
+	public static void LocalNotificationRemoveDetails(Context context, int notificationID)
+	{
+		SharedPreferences preferences = context.getSharedPreferences("LocalNotificationPreferences", MODE_PRIVATE);
+		SharedPreferences.Editor editor = preferences.edit();
+
+		try
+		{
+			JSONObject notificationDetails = new JSONObject(preferences.getString("notificationDetails", "{}"));
+			notificationDetails.remove(String.valueOf(notificationID));
+			editor.putString("notificationDetails", notificationDetails.toString());
+			editor.commit();
+		}
+		catch (JSONException e)
+		{
+			e.printStackTrace();
+		}
+	}
 
 	public int AndroidThunkJava_LocalNotificationScheduleAtTime(String targetDateTime, boolean localTime, String title, String body, String action, String activationEvent) 
 	{
 		int notificationID = LocalNotificationGetID(this);
 
+		if (!LocalNotificationScheduleAtTime(this, notificationID, targetDateTime, localTime, title, body, action, activationEvent))
+		{
+			return -1;
+		}
+
+		// Store notification details
+		SharedPreferences preferences = getSharedPreferences("LocalNotificationPreferences", MODE_PRIVATE);
+		SharedPreferences.Editor editor = preferences.edit();
+		JSONObject notificationDetails;
+		try
+		{
+			notificationDetails = new JSONObject(preferences.getString("notificationDetails", "{}"));
+		}
+		catch (JSONException e)
+		{
+			e.printStackTrace(); 
+			notificationDetails = new JSONObject();
+		}
+		JSONObject details = new JSONObject();
+		try
+		{
+			details.put("local-notification-targetDateTime", targetDateTime);
+			details.put("local-notification-localTime", localTime);
+			details.put("local-notification-title", title);
+			details.put("local-notification-body", body);
+			details.put("local-notification-action", action);
+			details.put("local-notification-activationEvent", activationEvent);
+			notificationDetails.put(String.valueOf(notificationID), details);
+			editor.putString("notificationDetails", notificationDetails.toString());
+			editor.commit();
+		}
+		catch (JSONException e)
+		{
+			e.printStackTrace();
+		}
+
+		return notificationID;
+	}
+
+	public static boolean LocalNotificationScheduleAtTime(Context context, int notificationID, String targetDateTime, boolean localTime, String title, String body, String action, String activationEvent)
+	{
 		// Create callback for PendingIntent
-		Intent notificationIntent = new Intent(this, LocalNotificationReceiver.class); 
+		Intent notificationIntent = new Intent(context, LocalNotificationReceiver.class);
 
 		// Add user-provided data
 		notificationIntent.putExtra("local-notification-ID", notificationID);
@@ -4733,9 +4792,9 @@ public class GameActivity extends $${gameActivitySuperClass}$$ implements Surfac
 		notificationIntent.putExtra("local-notification-body", body);
 		notificationIntent.putExtra("local-notification-action", action);
 		notificationIntent.putExtra("local-notification-activationEvent", activationEvent);
-		
+
 		// Designate the callback as a PendingIntent
-		PendingIntent pendingIntent = PendingIntent.getBroadcast(this, notificationID, notificationIntent, PendingIntent.FLAG_UPDATE_CURRENT);
+		PendingIntent pendingIntent = PendingIntent.getBroadcast(context, notificationID, notificationIntent, PendingIntent.FLAG_UPDATE_CURRENT);
 
 		TimeZone targetTimeZone = TimeZone.getTimeZone("UTC");
 
@@ -4757,8 +4816,8 @@ public class GameActivity extends $${gameActivitySuperClass}$$ implements Surfac
 		catch (ParseException e) 
 		{
 			e.printStackTrace();
-			LocalNotificationRemoveID(notificationID);
-			return -1;
+			LocalNotificationRemoveID(context, notificationID);
+			return false;
 		}
 
 		Date currentDate = new Date();
@@ -4767,17 +4826,17 @@ public class GameActivity extends $${gameActivitySuperClass}$$ implements Surfac
 
 		if(msDiff < 0)
 		{
-			LocalNotificationRemoveID(notificationID);
-			return -1;
+			LocalNotificationRemoveID(context, notificationID);
+			return false;
 		}
 
 		long futureTimeInMillis = SystemClock.elapsedRealtime() + msDiff;//Calculate the time to run the callback
-		AlarmManager alarmManager = (AlarmManager)getSystemService(Context.ALARM_SERVICE);
+		AlarmManager alarmManager = (AlarmManager)context.getSystemService(Context.ALARM_SERVICE);
 
 		//Schedule the operation by using AlarmService
 		alarmManager.set(AlarmManager.ELAPSED_REALTIME_WAKEUP, futureTimeInMillis, pendingIntent);
 
-		return notificationID;
+		return true;
 	}
 
 	public class LaunchNotification {
@@ -4814,6 +4873,7 @@ public class GameActivity extends $${gameActivitySuperClass}$$ implements Surfac
 		SharedPreferences preferences = getApplicationContext().getSharedPreferences("LocalNotificationPreferences", MODE_PRIVATE);
 		SharedPreferences.Editor editor = preferences.edit();
 		editor.putString("notificationIDs", "");
+		editor.putString("notificationDetails", "{}");
 		editor.commit();
 	}
 
@@ -4827,7 +4887,8 @@ public class GameActivity extends $${gameActivitySuperClass}$$ implements Surfac
 	{
 		if (AndroidThunkJava_LocalNotificationExists(notificationId))
 		{
-			LocalNotificationRemoveID(notificationId);
+			LocalNotificationRemoveID(this, notificationId);
+			LocalNotificationRemoveDetails(this, notificationId);
 
 			//Cancel the intent itself as well as from the alarm manager
 			AlarmManager alarmManager = (AlarmManager)getSystemService(Context.ALARM_SERVICE);
diff --git a/Engine/Build/Android/Java/src/com/epicgames/ue4/LocalNotificationReceiver.java b/Engine/Build/Android/Java/src/com/epicgames/ue4/LocalNotificationReceiver.java
--- a/Engine/Build/Android/Java/src/com/epicgames/ue4/LocalNotificationReceiver.java
+++ b/Engine/Build/Android/Java/src/com/epicgames/ue4/LocalNotificationReceiver.java
@@ -87,6 +87,9 @@ public class LocalNotificationReceiver extends BroadcastReceiver
 		{
 			// show the notification
 			notificationManager.notify(notificationID, notification);
+			
+			// clear the stored notification details if they exist
+			GameActivity.LocalNotificationRemoveDetails(context, notificationID);
 		}
 	}
 
diff --git a/Engine/Source/Programs/UnrealBuildTool/Platform/Android/UEDeployAndroid.cs b/Engine/Source/Programs/UnrealBuildTool/Platform/Android/UEDeployAndroid.cs
--- a/Engine/Source/Programs/UnrealBuildTool/Platform/Android/UEDeployAndroid.cs
+++ b/Engine/Source/Programs/UnrealBuildTool/Platform/Android/UEDeployAndroid.cs
@@ -2567,6 +2567,12 @@ namespace UnrealBuildTool
 
 			Text.AppendLine("\t\t<receiver android:name=\"com.epicgames.ue4.LocalNotificationReceiver\" />");
 
+			Text.AppendLine("\t\t<receiver android:name=\"com.epicgames.ue4.BootCompleteReceiver\" android:exported=\"true\">");
+			Text.AppendLine("\t\t\t<intent-filter>");
+			Text.AppendLine("\t\t\t\t<action android:name=\"android.intent.action.BOOT_COMPLETED\" />");
+			Text.AppendLine("\t\t\t</intent-filter>");
+			Text.AppendLine("\t\t</receiver>");
+
 			Text.AppendLine("\t\t<receiver android:name=\"com.epicgames.ue4.MulticastBroadcastReceiver\" android:exported=\"true\">");
 			Text.AppendLine("\t\t\t<intent-filter>");
 			Text.AppendLine("\t\t\t\t<action android:name=\"com.android.vending.INSTALL_REFERRER\" />");
@@ -2606,6 +2612,7 @@ namespace UnrealBuildTool
 			//	Text.AppendLine("\t<uses-permission android:name=\"android.permission.READ_PHONE_STATE\"/>");
 				Text.AppendLine("\t<uses-permission android:name=\"com.android.vending.CHECK_LICENSE\"/>");
 				Text.AppendLine("\t<uses-permission android:name=\"android.permission.ACCESS_WIFI_STATE\"/>");
+				Text.AppendLine("\t<uses-permission android:name=\"android.permission.RECEIVE_BOOT_COMPLETED\"/>");
 
 				if (bEnableGooglePlaySupport && bUseGetAccounts)
 				{
