Cancel notification if already displayed when calling CancelNotification on Android
#rnx

diff --git a/Engine/Build/Android/Java/src/com/epicgames/ue4/GameActivity.java.template b/Engine/Build/Android/Java/src/com/epicgames/ue4/GameActivity.java.template
--- a/Engine/Build/Android/Java/src/com/epicgames/ue4/GameActivity.java.template
+++ b/Engine/Build/Android/Java/src/com/epicgames/ue4/GameActivity.java.template
@@ -3,6 +3,7 @@
 
 package com.epicgames.ue4;
 
+import android.app.NotificationManager;
 import android.opengl.EGL14;
 import android.opengl.EGLConfig;
 import android.opengl.EGLContext;
@@ -4895,6 +4896,10 @@ public class GameActivity extends $${gameActivitySuperClass}$$ implements Surfac
 			PendingIntent pendingIntent = PendingIntent.getBroadcast(this, notificationId, new Intent(this, LocalNotificationReceiver.class), PendingIntent.FLAG_UPDATE_CURRENT);
 			pendingIntent.cancel();
 			alarmManager.cancel(pendingIntent);
+			
+			//Also clear the notification if currently displayed
+			NotificationManager notificationManager = (NotificationManager)getSystemService(Context.NOTIFICATION_SERVICE);
+			notificationManager.cancel(notificationId);
 			return true;
 		}
 		return false;
