Use standard year range for copyright notice, to fix CIS warning.

#rnx

diff --git a/Engine/Build/Android/Java/src/com/epicgames/ue4/BootCompleteReceiver.java b/Engine/Build/Android/Java/src/com/epicgames/ue4/BootCompleteReceiver.java
--- a/Engine/Build/Android/Java/src/com/epicgames/ue4/BootCompleteReceiver.java
+++ b/Engine/Build/Android/Java/src/com/epicgames/ue4/BootCompleteReceiver.java
@@ -1,4 +1,4 @@
-// Copyright 2019 Epic Games, Inc. All Rights Reserved.
+// Copyright 1998-2019 Epic Games, Inc. All Rights Reserved.
 
 package com.epicgames.ue4;
 
