Scheduled local notifications on Android will now persist across reboots like on iOS.
Commiting changes from post-commit review.
#rnx

diff --git a/Engine/Build/Android/Java/src/com/epicgames/ue4/BootCompleteReceiver.java b/Engine/Build/Android/Java/src/com/epicgames/ue4/BootCompleteReceiver.java
--- a/Engine/Build/Android/Java/src/com/epicgames/ue4/BootCompleteReceiver.java
+++ b/Engine/Build/Android/Java/src/com/epicgames/ue4/BootCompleteReceiver.java
@@ -6,59 +6,109 @@ import android.content.BroadcastReceiver;
 import android.content.Context;
 import android.content.Intent;
 import android.content.SharedPreferences;
+import android.os.AsyncTask;
 
 import org.json.JSONException;
 import org.json.JSONObject;
 
+import java.lang.ref.WeakReference;
 import java.util.Iterator;
 
 public class BootCompleteReceiver extends BroadcastReceiver
 {
+	public static Logger Log = new Logger("UE4", "BootCompleteReceiver");
+	
 	@Override
 	public void onReceive(Context context, Intent intent)
 	{
-		// restore any scheduled notifications
-		SharedPreferences preferences = context.getSharedPreferences("LocalNotificationPreferences", Context.MODE_PRIVATE);
-		SharedPreferences.Editor editor = preferences.edit();
-		try
+		PendingResult pendingResult = goAsync();
+		Task task = new Task(pendingResult, context);
+		task.execute();
+	}
+
+	private static class Task extends AsyncTask<Void, Void, Void>
+	{
+		private PendingResult pendingResult;
+		private WeakReference<Context> context;
+
+		Task(PendingResult pendingResult, Context context)
 		{
-			boolean changed = false;
-			JSONObject notificationDetails = new JSONObject(preferences.getString("notificationDetails", "{}"));
-			for (Iterator<String> i = notificationDetails.keys(); i.hasNext(); )
+			this.pendingResult = pendingResult;
+			// prevent async task context leaks
+			this.context = new WeakReference<>(context.getApplicationContext());
+		}
+
+		@Override
+		protected Void doInBackground(Void... params)
+		{
+			Context context = this.context.get();
+			if (context == null)
+			{
+				// should technically never happen, since the application context will still exist until the receiver is recycled
+				return null;
+			}
+
+			// restore any scheduled notifications
+			SharedPreferences preferences = context.getSharedPreferences("LocalNotificationPreferences", Context.MODE_PRIVATE);
+			if (preferences == null)
+			{
+				// nothing to reschedule if the preferences don't exist
+				return null;
+			}
+			try
 			{
-				try
+				boolean changed = false;
+				JSONObject notificationDetails = new JSONObject(preferences.getString("notificationDetails", "{}"));
+				for (Iterator<String> iterator = notificationDetails.keys(); iterator.hasNext(); )
 				{
-					String key = i.next();
-					int notificationId = Integer.parseInt(key); 
-					JSONObject details = notificationDetails.getJSONObject(key);
-					String targetDateTime = details.getString("local-notification-targetDateTime");
-					boolean localTime = details.getBoolean("local-notification-localTime");
-					String title = details.getString("local-notification-title");
-					String body = details.getString("local-notification-body");
-					String action = details.getString("local-notification-action");
-					String activationEvent = details.getString("local-notification-activationEvent");
-					
-					if (!GameActivity.LocalNotificationScheduleAtTime(context, notificationId, targetDateTime, localTime, title, body, action, activationEvent))
+					try
 					{
-						// if it fails, remove from details
-						i.remove();
+						String key = iterator.next();
+						int notificationId = Integer.parseInt(key);
+						JSONObject details = notificationDetails.getJSONObject(key);
+						String targetDateTime = details.getString("local-notification-targetDateTime");
+						boolean localTime = details.getBoolean("local-notification-localTime");
+						String title = details.getString("local-notification-title");
+						String body = details.getString("local-notification-body");
+						String action = details.getString("local-notification-action");
+						String activationEvent = details.getString("local-notification-activationEvent");
+
+						if (!GameActivity.LocalNotificationScheduleAtTime(context, notificationId, targetDateTime, localTime, title, body, action, activationEvent))
+						{
+							// if it fails, remove from details
+							iterator.remove();
+							changed = true;
+						}
+					}
+					catch (NumberFormatException | JSONException e)
+					{
+						Log.error("Error reading notification details", e);
+						iterator.remove();
+						changed = true;
 					}
 				}
-				catch (NumberFormatException | JSONException e)
+
+				if (changed)
 				{
-					e.printStackTrace();
-					i.remove();
+					SharedPreferences.Editor editor = preferences.edit();
+					editor.putString("notificationDetails", notificationDetails.toString());
+					editor.commit();
 				}
 			}
-			
-			if (changed)
+			catch (JSONException e)
 			{
-				editor.commit();
+				Log.error("Error reading notification details", e);
 			}
+
+			return null;
 		}
-		catch (JSONException e)
+
+		@Override
+		protected void onPostExecute(Void result)
 		{
-			e.printStackTrace();
+			super.onPostExecute(result);
+			// mark receiver as ready for recycle
+			pendingResult.finish();
 		}
 	}
 }
diff --git a/Engine/Config/BaseEngine.ini b/Engine/Config/BaseEngine.ini
--- a/Engine/Config/BaseEngine.ini
+++ b/Engine/Config/BaseEngine.ini
@@ -1979,6 +1979,7 @@ bAllowControllers=True
 bAllowIMU=True
 bUseDisplayCutout=False
 bEnableSnapshots=False
+bRestoreNotificationsOnReboot=False
 
 [/Script/AndroidPlatformEditor.AndroidSDKSettings]
 SDKAPILevel=latest
diff --git a/Engine/Source/Programs/UnrealBuildTool/Platform/Android/UEDeployAndroid.cs b/Engine/Source/Programs/UnrealBuildTool/Platform/Android/UEDeployAndroid.cs
--- a/Engine/Source/Programs/UnrealBuildTool/Platform/Android/UEDeployAndroid.cs
+++ b/Engine/Source/Programs/UnrealBuildTool/Platform/Android/UEDeployAndroid.cs
@@ -2210,6 +2210,8 @@ namespace UnrealBuildTool
 			Ini.GetBool("/Script/AndroidRuntimeSettings.AndroidRuntimeSettings", "bFullScreen", out EnableFullScreen);
 			bool bUseDisplayCutout;
 			Ini.GetBool("/Script/AndroidRuntimeSettings.AndroidRuntimeSettings", "bUseDisplayCutout", out bUseDisplayCutout);
+			bool bRestoreNotificationsOnReboot = false;
+			Ini.GetBool("/Script/AndroidRuntimeSettings.AndroidRuntimeSettings", "bRestoreNotificationsOnReboot", out bRestoreNotificationsOnReboot);
 			List<string> ExtraManifestNodeTags;
 			Ini.GetArray("/Script/AndroidRuntimeSettings.AndroidRuntimeSettings", "ExtraManifestNodeTags", out ExtraManifestNodeTags);
 			List<string> ExtraApplicationNodeTags;
@@ -2567,11 +2569,16 @@ namespace UnrealBuildTool
 
 			Text.AppendLine("\t\t<receiver android:name=\"com.epicgames.ue4.LocalNotificationReceiver\" />");
 
-			Text.AppendLine("\t\t<receiver android:name=\"com.epicgames.ue4.BootCompleteReceiver\" android:exported=\"true\">");
-			Text.AppendLine("\t\t\t<intent-filter>");
-			Text.AppendLine("\t\t\t\t<action android:name=\"android.intent.action.BOOT_COMPLETED\" />");
-			Text.AppendLine("\t\t\t</intent-filter>");
-			Text.AppendLine("\t\t</receiver>");
+			if (bRestoreNotificationsOnReboot)
+			{
+				Text.AppendLine("\t\t<receiver android:name=\"com.epicgames.ue4.BootCompleteReceiver\" android:exported=\"true\">");
+				Text.AppendLine("\t\t\t<intent-filter>");
+				Text.AppendLine("\t\t\t\t<action android:name=\"android.intent.action.BOOT_COMPLETED\" />");
+				Text.AppendLine("\t\t\t\t<action android:name=\"android.intent.action.QUICKBOOT_POWERON\" />");
+				Text.AppendLine("\t\t\t\t<action android:name=\"com.htc.intent.action.QUICKBOOT_POWERON\" />");
+				Text.AppendLine("\t\t\t</intent-filter>");
+				Text.AppendLine("\t\t</receiver>");
+			}
 
 			Text.AppendLine("\t\t<receiver android:name=\"com.epicgames.ue4.MulticastBroadcastReceiver\" android:exported=\"true\">");
 			Text.AppendLine("\t\t\t<intent-filter>");
@@ -2612,7 +2619,11 @@ namespace UnrealBuildTool
 			//	Text.AppendLine("\t<uses-permission android:name=\"android.permission.READ_PHONE_STATE\"/>");
 				Text.AppendLine("\t<uses-permission android:name=\"com.android.vending.CHECK_LICENSE\"/>");
 				Text.AppendLine("\t<uses-permission android:name=\"android.permission.ACCESS_WIFI_STATE\"/>");
-				Text.AppendLine("\t<uses-permission android:name=\"android.permission.RECEIVE_BOOT_COMPLETED\"/>");
+
+				if (bRestoreNotificationsOnReboot)
+				{
+					Text.AppendLine("\t<uses-permission android:name=\"android.permission.RECEIVE_BOOT_COMPLETED\"/>");
+				}
 
 				if (bEnableGooglePlaySupport && bUseGetAccounts)
 				{
diff --git a/Engine/Source/Runtime/Android/AndroidRuntimeSettings/Classes/AndroidRuntimeSettings.h b/Engine/Source/Runtime/Android/AndroidRuntimeSettings/Classes/AndroidRuntimeSettings.h
--- a/Engine/Source/Runtime/Android/AndroidRuntimeSettings/Classes/AndroidRuntimeSettings.h
+++ b/Engine/Source/Runtime/Android/AndroidRuntimeSettings/Classes/AndroidRuntimeSettings.h
@@ -271,6 +271,10 @@ public:
 	UPROPERTY(GlobalConfig, EditAnywhere, Category = "APK Packaging", Meta = (DisplayName = "Use display cutout region?"))
 	bool bUseDisplayCutout;
 
+	// Should we restore scheduled local notifications on reboot? This will add a receiver for boot complete and a permission to the manifest.
+	UPROPERTY(GlobalConfig, EditAnywhere, Category = "APK Packaging", Meta = (DisplayName = "Restore scheduled notifications on reboot"))
+	bool bRestoreNotificationsOnReboot;
+
 	// Level of verbosity to use during packaging with Ant
 	UPROPERTY(GlobalConfig, EditAnywhere, Category = "APK Packaging")
 	TEnumAsByte<EAndroidAntVerbosity::Type> AntVerbosity;
