Add notification channel for Android 26 compatibility

#ue4
#android

diff --git a/Engine/Build/Android/Java/src/com/epicgames/ue4/LocalNotificationReceiver.java b/Engine/Build/Android/Java/src/com/epicgames/ue4/LocalNotificationReceiver.java
--- a/Engine/Build/Android/Java/src/com/epicgames/ue4/LocalNotificationReceiver.java
+++ b/Engine/Build/Android/Java/src/com/epicgames/ue4/LocalNotificationReceiver.java
@@ -4,7 +4,8 @@
  
  import android.app.Notification;
  import android.app.NotificationManager;
- import android.app.PendingIntent;
+ import android.app.NotificationChannel;
+ import android.app.PendingIntent; 
  import android.content.BroadcastReceiver;
  import android.content.Context;
  import android.content.Intent;
@@ -12,7 +13,10 @@
  
  public class LocalNotificationReceiver extends BroadcastReceiver
  {
- 
+	 private static boolean bChannelExists = false;
+	 private static final String NOTIFICATION_CHANNEL_ID = "ue4-push-notification-channel-id";
+	 private static final CharSequence NOTICATION_CHANNEL_NAME = "ue4-push-notification-channel";
+	 
 	 public void onReceive(Context context, Intent intent)
 	 {
 		int notificationID = intent.getIntExtra("local-notification-ID" , 0);
@@ -43,15 +47,37 @@
 		}
 		PendingIntent pendingNotificationIntent = PendingIntent.getActivity(context, notificationID, notificationIntent, 0);
 
-		NotificationCompat.Builder builder = new NotificationCompat.Builder(context)
+		if (android.os.Build.VERSION.SDK_INT >= 26)
+		{
+			NotificationCompat.Builder builder = new NotificationCompat.Builder(context, NOTIFICATION_CHANNEL_ID);
 			.setSmallIcon(notificationIconID)
 			.setContentIntent(pendingNotificationIntent)
 			.setWhen(System.currentTimeMillis())
 			.setTicker(details)		// note: will not show up on Lollipop up except for accessibility
-			.setContentTitle(title);
-		if (android.os.Build.VERSION.SDK_INT >= 21)
+			.setContentTitle(title)
+			.setColor(0xff0e1e43);
+			
+			if (!bChannelExists)
+			{
+				NotificationChannel channel = new NotificationChannel(NOTIFICATION_CHANNEL_ID, NOTICATION_CHANNEL_NAME, NotificationManager.IMPORTANCE_DEFAULT);
+				channel.enableVibration(true);
+				channel.enableLights(true);
+				manager.createNotificationChannel(channel);
+				bChannelExists = true;			}
+			}
+		}
+		else
 		{
-			builder.setColor(0xff0e1e43);
+			NotificationCompat.Builder builder = new NotificationCompat.Builder(context)
+			.setSmallIcon(notificationIconID)
+			.setContentIntent(pendingNotificationIntent)
+			.setWhen(System.currentTimeMillis())
+			.setTicker(details)		// note: will not show up on Lollipop up except for accessibility
+			.setContentTitle(title);
+			if (android.os.Build.VERSION.SDK_INT >= 21)
+			{
+				builder.setColor(0xff0e1e43);
+			}
 		}
 		Notification notification = builder.build();
 
