Fix bad merge

#ue4
#android

diff --git a/Engine/Build/Android/Java/src/com/epicgames/ue4/LocalNotificationReceiver.java b/Engine/Build/Android/Java/src/com/epicgames/ue4/LocalNotificationReceiver.java
--- a/Engine/Build/Android/Java/src/com/epicgames/ue4/LocalNotificationReceiver.java
+++ b/Engine/Build/Android/Java/src/com/epicgames/ue4/LocalNotificationReceiver.java
@@ -47,9 +47,12 @@
 		}
 		PendingIntent pendingNotificationIntent = PendingIntent.getActivity(context, notificationID, notificationIntent, 0);
 
+		NotificationManager notificationManager = (NotificationManager) context.getSystemService(Context.NOTIFICATION_SERVICE);
+		NotificationCompat.Builder builder = null;
+
 		if (android.os.Build.VERSION.SDK_INT >= 26)
 		{
-			NotificationCompat.Builder builder = new NotificationCompat.Builder(context, NOTIFICATION_CHANNEL_ID);
+			builder = new NotificationCompat.Builder(context, NOTIFICATION_CHANNEL_ID)
 			.setSmallIcon(notificationIconID)
 			.setContentIntent(pendingNotificationIntent)
 			.setWhen(System.currentTimeMillis())
@@ -62,13 +65,13 @@
 				NotificationChannel channel = new NotificationChannel(NOTIFICATION_CHANNEL_ID, NOTICATION_CHANNEL_NAME, NotificationManager.IMPORTANCE_DEFAULT);
 				channel.enableVibration(true);
 				channel.enableLights(true);
-				manager.createNotificationChannel(channel);
-				bChannelExists = true;			}
+				notificationManager.createNotificationChannel(channel);
+				bChannelExists = true;
 			}
 		}
 		else
 		{
-			NotificationCompat.Builder builder = new NotificationCompat.Builder(context)
+			builder = new NotificationCompat.Builder(context)
 			.setSmallIcon(notificationIconID)
 			.setContentIntent(pendingNotificationIntent)
 			.setWhen(System.currentTimeMillis())
@@ -85,8 +88,6 @@
 		notification.flags |= Notification.FLAG_AUTO_CANCEL;
 		notification.defaults |= Notification.DEFAULT_SOUND | Notification.DEFAULT_VIBRATE;
 
-		NotificationManager notificationManager = (NotificationManager) context.getSystemService(Context.NOTIFICATION_SERVICE);
-
 		// show the notification
 		notificationManager.notify(notificationID, notification); 
 	 }
