Fix ConfigRules.jar parsing for Linux line endings

#ue4
#android

diff --git a/Engine/Build/Android/Prebuilt/ConfigRulesTool/classes/placeholder.txt b/Engine/Build/Android/Prebuilt/ConfigRulesTool/classes/placeholder.txt
new file mode 100644
--- /dev/null
+++ b/Engine/Build/Android/Prebuilt/ConfigRulesTool/classes/placeholder.txt
@@ -0,0 +1 @@
+This file is here to create the directory.
\ No newline at end of file
diff --git a/Engine/Build/Android/Prebuilt/ConfigRulesTool/compile.sh b/Engine/Build/Android/Prebuilt/ConfigRulesTool/compile.sh
new file mode 100755
--- /dev/null
+++ b/Engine/Build/Android/Prebuilt/ConfigRulesTool/compile.sh
@@ -0,0 +1,4 @@
+#!/bin/bash
+javac -d classes src/*.java
+cd classes
+jar cfe ../bin/ConfigRulesTool.jar ConfigRulesTool *.class
diff --git a/Engine/Build/Android/Prebuilt/ConfigRulesTool/src/ConfigRulesTool.java b/Engine/Build/Android/Prebuilt/ConfigRulesTool/src/ConfigRulesTool.java
--- a/Engine/Build/Android/Prebuilt/ConfigRulesTool/src/ConfigRulesTool.java
+++ b/Engine/Build/Android/Prebuilt/ConfigRulesTool/src/ConfigRulesTool.java
@@ -97,6 +97,8 @@ public class ConfigRulesTool
 				{
 					int eolIndex = versionLine.indexOf("\r");
 					int newIndex = versionLine.indexOf("\n");
+					eolIndex = (eolIndex < 0) ? 1000 : eolIndex;
+					newIndex = (newIndex < 0) ? 1000 : newIndex;
 					eolIndex = (eolIndex < newIndex) ? eolIndex : newIndex;
 					try
 					{
