Fix for LocalNotificationReceiver using Ant on projects using older dependencies

#ue4
#android

diff --git a/Engine/Build/Android/Java/src/com/epicgames/ue4/LocalNotificationReceiver.java b/Engine/Build/Android/Java/src/com/epicgames/ue4/LocalNotificationReceiver.java
--- a/Engine/Build/Android/Java/src/com/epicgames/ue4/LocalNotificationReceiver.java
+++ b/Engine/Build/Android/Java/src/com/epicgames/ue4/LocalNotificationReceiver.java
@@ -48,18 +48,19 @@
 		PendingIntent pendingNotificationIntent = PendingIntent.getActivity(context, notificationID, notificationIntent, 0);
 
 		NotificationManager notificationManager = (NotificationManager) context.getSystemService(Context.NOTIFICATION_SERVICE);
-		NotificationCompat.Builder builder = null;
-
-		if (android.os.Build.VERSION.SDK_INT >= 26)
-		{
-			builder = new NotificationCompat.Builder(context, NOTIFICATION_CHANNEL_ID)
+		NotificationCompat.Builder builder = new NotificationCompat.Builder(context, NOTIFICATION_CHANNEL_ID)
 			.setSmallIcon(notificationIconID)
 			.setContentIntent(pendingNotificationIntent)
 			.setWhen(System.currentTimeMillis())
 			.setTicker(details)		// note: will not show up on Lollipop up except for accessibility
-			.setContentTitle(title)
-			.setColor(0xff0e1e43);
-			
+			.setContentTitle(title);
+		if (android.os.Build.VERSION.SDK_INT >= 21)
+		{
+			builder.setColor(0xff0e1e43);
+		}
+
+		if (android.os.Build.VERSION.SDK_INT >= 26)
+		{
 			if (!bChannelExists)
 			{
 				NotificationChannel channel = new NotificationChannel(NOTIFICATION_CHANNEL_ID, NOTICATION_CHANNEL_NAME, NotificationManager.IMPORTANCE_DEFAULT);
@@ -69,19 +70,6 @@
 				bChannelExists = true;
 			}
 		}
-		else
-		{
-			builder = new NotificationCompat.Builder(context)
-			.setSmallIcon(notificationIconID)
-			.setContentIntent(pendingNotificationIntent)
-			.setWhen(System.currentTimeMillis())
-			.setTicker(details)		// note: will not show up on Lollipop up except for accessibility
-			.setContentTitle(title);
-			if (android.os.Build.VERSION.SDK_INT >= 21)
-			{
-				builder.setColor(0xff0e1e43);
-			}
-		}
 		Notification notification = builder.build();
 
 		// Stick with the defaults
