Do not call getPresentationDeadlineNanos on older Android versions

#ue4
#android

diff --git a/Engine/Build/Android/Java/src/com/epicgames/ue4/GameActivity.java.template b/Engine/Build/Android/Java/src/com/epicgames/ue4/GameActivity.java.template
--- a/Engine/Build/Android/Java/src/com/epicgames/ue4/GameActivity.java.template
+++ b/Engine/Build/Android/Java/src/com/epicgames/ue4/GameActivity.java.template
@@ -4724,9 +4724,13 @@ public class GameActivity extends NativeActivity implements SurfaceHolder.Callba
 	{
 		if (key.equals("ue4.display.PresentationDeadlineNanos"))
 		{
-			WindowManager windowManager = (WindowManager)getSystemService(WINDOW_SERVICE);
-			android.view.Display display = windowManager.getDefaultDisplay();
-			return display.getPresentationDeadlineNanos();
+			if (android.os.Build.VERSION.SDK_INT >= 21)
+			{
+				WindowManager windowManager = (WindowManager)getSystemService(WINDOW_SERVICE);
+				android.view.Display display = windowManager.getDefaultDisplay();
+				return display.getPresentationDeadlineNanos();
+			}
+			return -1;
 		}
 		else
 		if (key.equals("ue4.display.AppVsyncOffsetNanos"))
