Fix TextureFormats variable in configrules

#ue4
#android

diff --git a/Engine/Build/Android/Java/src/com/epicgames/ue4/GameActivity.java.template b/Engine/Build/Android/Java/src/com/epicgames/ue4/GameActivity.java.template
--- a/Engine/Build/Android/Java/src/com/epicgames/ue4/GameActivity.java.template
+++ b/Engine/Build/Android/Java/src/com/epicgames/ue4/GameActivity.java.template
@@ -1109,7 +1109,7 @@ public class GameActivity extends NativeActivity implements SurfaceHolder.Callba
 		}
 		if (Extensions.contains("GL_IMG_texture_compression_pvrtc"))
 		{
-			TextureFormats += ",PVRTC";
+			TextureFormats += "PVRTC,";
 		}
 		if (Extensions.contains("GL_NV_texture_compression_s3tc") || Extensions.contains("GL_EXT_texture_compression_s3tc"))
 		{
