Fix for Ant NotificationCompat issue

#ue4
#android

diff --git a/Engine/Build/Android/Legacy/src/com/epicgames/ue4/LocalNotificationReceiver.java b/Engine/Build/Android/Legacy/src/com/epicgames/ue4/LocalNotificationReceiver.java
new file mode 100644
--- /dev/null
+++ b/Engine/Build/Android/Legacy/src/com/epicgames/ue4/LocalNotificationReceiver.java
@@ -0,0 +1,68 @@
+// Copyright 1998-2018 Epic Games, Inc. All Rights Reserved.
+
+ package com.epicgames.ue4;
+ 
+ import android.app.Notification;
+ import android.app.NotificationManager;
+ import android.app.PendingIntent;
+ import android.content.BroadcastReceiver;
+ import android.content.Context;
+ import android.content.Intent;
+ import android.support.v4.app.NotificationCompat;
+ 
+ public class LocalNotificationReceiver extends BroadcastReceiver
+ {
+ 
+	 public void onReceive(Context context, Intent intent)
+	 {
+		int notificationID = intent.getIntExtra("local-notification-ID" , 0);
+		String title  = intent.getStringExtra("local-notification-title");
+		String details  = intent.getStringExtra("local-notification-body");
+		String action = intent.getStringExtra("local-notification-action");
+		String activationEvent = intent.getStringExtra("local-notification-activationEvent");
+
+		if(title == null || details == null || action == null || activationEvent == null)
+		{
+			// Do not schedule any local notification if any allocation failed
+			return;
+		}
+
+		// Open UE4 app if clicked
+		Intent notificationIntent = new Intent(context, GameActivity.class);
+
+		// launch if closed but resume if running
+		notificationIntent.setFlags(Intent.FLAG_ACTIVITY_CLEAR_TOP | Intent.FLAG_ACTIVITY_SINGLE_TOP);
+		notificationIntent.putExtra("localNotificationID" , notificationID);
+		notificationIntent.putExtra("localNotificationAppLaunched" , true);
+		notificationIntent.putExtra("localNotificationLaunchActivationEvent", activationEvent);
+
+		int notificationIconID = context.getResources().getIdentifier("ic_notification", "drawable", context.getPackageName());
+		if (notificationIconID == 0)
+		{
+			notificationIconID = context.getResources().getIdentifier("icon", "drawable", context.getPackageName());
+		}
+		PendingIntent pendingNotificationIntent = PendingIntent.getActivity(context, notificationID, notificationIntent, 0);
+
+		NotificationCompat.Builder builder = new NotificationCompat.Builder(context)
+			.setSmallIcon(notificationIconID)
+			.setContentIntent(pendingNotificationIntent)
+			.setWhen(System.currentTimeMillis())
+			.setTicker(details)		// note: will not show up on Lollipop up except for accessibility
+			.setContentTitle(title);
+		if (android.os.Build.VERSION.SDK_INT >= 21)
+		{
+			builder.setColor(0xff0e1e43);
+		}
+		Notification notification = builder.build();
+
+		// Stick with the defaults
+		notification.flags |= Notification.FLAG_AUTO_CANCEL;
+		notification.defaults |= Notification.DEFAULT_SOUND | Notification.DEFAULT_VIBRATE;
+
+		NotificationManager notificationManager = (NotificationManager) context.getSystemService(Context.NOTIFICATION_SERVICE);
+
+		// show the notification
+		notificationManager.notify(notificationID, notification); 
+	 }
+ }
+ 
\ No newline at end of file
diff --git a/Engine/Source/Programs/UnrealBuildTool/Platform/Android/UEDeployAndroid.cs b/Engine/Source/Programs/UnrealBuildTool/Platform/Android/UEDeployAndroid.cs
--- a/Engine/Source/Programs/UnrealBuildTool/Platform/Android/UEDeployAndroid.cs
+++ b/Engine/Source/Programs/UnrealBuildTool/Platform/Android/UEDeployAndroid.cs
@@ -3331,6 +3331,10 @@ namespace UnrealBuildTool
 			CopyFileDirectory(UE4BuildFilesPath, UE4BuildPath, Replacements);
 			CopyFileDirectory(UE4BuildFilesPath + "/NotForLicensees", UE4BuildPath, Replacements);
 			CopyFileDirectory(UE4BuildFilesPath + "/NoRedist", UE4BuildPath, Replacements);
+			if (!bGradleEnabled)
+			{
+				CopyFileDirectory(Path.Combine(EngineDirectory, "Build", "Android", "Legacy"), UE4BuildPath, Replacements);
+			}
 			CopyFileDirectory(GameBuildFilesPath, UE4BuildPath, Replacements);
 			CopyFileDirectory(GameBuildFilesPath + "/NotForLicensees", UE4BuildPath, Replacements);
 			CopyFileDirectory(GameBuildFilesPath + "/NoRedist", UE4BuildPath, Replacements);
