Added support for using display cutout region on Android
- added new project setting bUseDisplayCutout
- temporary restriction to HUAWEI devices for this release (remove later)

#android

diff --git a/Engine/Build/Android/Java/src/com/epicgames/ue4/GameActivity.java.template b/Engine/Build/Android/Java/src/com/epicgames/ue4/GameActivity.java.template
--- a/Engine/Build/Android/Java/src/com/epicgames/ue4/GameActivity.java.template
+++ b/Engine/Build/Android/Java/src/com/epicgames/ue4/GameActivity.java.template
@@ -121,6 +121,7 @@ import android.webkit.CookieSyncManager;
 import android.media.AudioManager;
 import android.util.DisplayMetrics;
 
+import android.view.DisplayCutout;
 import android.view.InputDevice;
 import android.view.KeyEvent;
 import android.view.Gravity;
@@ -133,6 +134,7 @@ import android.view.ViewGroup.MarginLayoutParams;
 import android.view.Surface;
 import android.view.SurfaceHolder;
 import android.view.SurfaceView;
+import android.view.WindowInsets;
 import android.view.WindowManager;
 import android.view.Window;
 import android.widget.LinearLayout;
@@ -412,6 +414,8 @@ public class GameActivity extends NativeActivity implements SurfaceHolder.Callba
 	private String ForceExitUpdateButtonText = "";
 	private String ForceExitLink = "";
 	
+	private boolean SplashScreenLaunch = false;
+	private boolean UseDisplayCutout = false;
 	private boolean ShouldHideUI = false;
 
 	/** Whether this application is for distribution */
@@ -1948,8 +1952,10 @@ public class GameActivity extends NativeActivity implements SurfaceHolder.Callba
 		if (_extrasBundle != null)
 		{
 			ShouldHideUI = _extrasBundle.getString("ShouldHideUI") != null;
+			UseDisplayCutout = _extrasBundle.getString("UseDisplayCutout") != null;
 			if (_extrasBundle.getString("UseSplashScreen") != null)
 			{
+				SplashScreenLaunch = true;
 				try {
 					// try to get the splash theme (can't use R.style.UE4SplashTheme since we don't know the package name until runtime)
 					int SplashThemeId = getResources().getIdentifier("UE4SplashTheme", "style", getPackageName());
@@ -1990,6 +1996,14 @@ public class GameActivity extends NativeActivity implements SurfaceHolder.Callba
 			catch (Exception e) {
 				e.printStackTrace();
 			}
+
+			if (UseDisplayCutout)
+			{
+				// will not be true if not Android Pie or later
+	            WindowManager.LayoutParams params = getWindow().getAttributes();
+		        params.layoutInDisplayCutoutMode = WindowManager.LayoutParams.LAYOUT_IN_DISPLAY_CUTOUT_MODE_SHORT_EDGES;
+			    getWindow().setAttributes(params);
+			}
 		}
 
 		//Check for target sdk.  If 23 or higher then warn that permission handling may mean features don't work if user denies them.
@@ -2299,7 +2313,7 @@ public class GameActivity extends NativeActivity implements SurfaceHolder.Callba
 				VerifyOBBOnStartUp = false;
 				Log.debug( "Did not find bVerifyOBBOnStartUp, using default.");
 			}
-				
+
 			if(bundle.containsKey("com.epicgames.ue4.GameActivity.bShouldHideUI"))
 			{
 				ShouldHideUI = bundle.getBoolean("com.epicgames.ue4.GameActivity.bShouldHideUI");
@@ -2309,6 +2323,20 @@ public class GameActivity extends NativeActivity implements SurfaceHolder.Callba
 			{
 				Log.debug( "UI hiding not found. Leaving as " + ShouldHideUI);
 			}
+
+			if (SplashScreenLaunch == false && android.os.Build.VERSION.SDK_INT >= 28)
+			{
+				if(bundle.containsKey("com.epicgames.ue4.GameActivity.bUseDisplayCutout"))
+				{
+					UseDisplayCutout = bundle.getBoolean("com.epicgames.ue4.GameActivity.bUseDisplayCutout");
+					Log.debug( "Display cutout set to " + UseDisplayCutout);
+				}
+				else
+				{
+					Log.debug( "Display cutout not found. Leaving as " + UseDisplayCutout);
+				}
+			}
+
 			if(bundle.containsKey("com.epicgames.ue4.GameActivity.BuildConfiguration"))
 			{
 				BuildConfiguration = bundle.getString("com.epicgames.ue4.GameActivity.BuildConfiguration");
@@ -2797,6 +2825,18 @@ public class GameActivity extends NativeActivity implements SurfaceHolder.Callba
 			});
 		}
 
+		// only true if Android Pie or later
+		if (UseDisplayCutout)
+		{
+			// only on Android Pie and later
+			WindowManager.LayoutParams params = getWindow().getAttributes();
+			if (params.layoutInDisplayCutoutMode != WindowManager.LayoutParams.LAYOUT_IN_DISPLAY_CUTOUT_MODE_SHORT_EDGES)
+			{
+				params.layoutInDisplayCutoutMode = WindowManager.LayoutParams.LAYOUT_IN_DISPLAY_CUTOUT_MODE_SHORT_EDGES;
+				getWindow().setAttributes(params);
+			}
+		}
+
 		if(HasAllFiles)
 		{
 			Log.debug("==============> Resuming main init");
@@ -2875,6 +2915,17 @@ public class GameActivity extends NativeActivity implements SurfaceHolder.Callba
 			});
 		}
 
+		// only true if Android Pie or later
+		if (UseDisplayCutout)
+		{
+			WindowManager.LayoutParams params = getWindow().getAttributes();
+			if (params.layoutInDisplayCutoutMode != WindowManager.LayoutParams.LAYOUT_IN_DISPLAY_CUTOUT_MODE_SHORT_EDGES)
+			{
+				params.layoutInDisplayCutoutMode = WindowManager.LayoutParams.LAYOUT_IN_DISPLAY_CUTOUT_MODE_SHORT_EDGES;
+				getWindow().setAttributes(params);
+			}
+		}
+
 		// restore screensaver state
 		AndroidThunkJava_KeepScreenOn(bKeepScreenOn);
 		
diff --git a/Engine/Build/Android/Java/src/com/epicgames/ue4/SplashActivity.java b/Engine/Build/Android/Java/src/com/epicgames/ue4/SplashActivity.java
--- a/Engine/Build/Android/Java/src/com/epicgames/ue4/SplashActivity.java
+++ b/Engine/Build/Android/Java/src/com/epicgames/ue4/SplashActivity.java
@@ -9,6 +9,7 @@ import android.content.pm.ApplicationInfo;
 import android.content.pm.PackageManager;
 import android.content.pm.PackageManager.NameNotFoundException;
 import android.view.View;
+import android.view.WindowManager;
 
 public class SplashActivity extends Activity
 {
@@ -18,6 +19,7 @@ public class SplashActivity extends Activity
 		super.onCreate(savedInstanceState);
 
 		boolean ShouldHideUI = false;
+		boolean UseDisplayCutout = false;
 		try {
 			ApplicationInfo ai = getPackageManager().getApplicationInfo(getPackageName(), PackageManager.GET_META_DATA);
 			Bundle bundle = ai.metaData;
@@ -26,6 +28,10 @@ public class SplashActivity extends Activity
 			{
 				ShouldHideUI = bundle.getBoolean("com.epicgames.ue4.GameActivity.bShouldHideUI");
 			}
+			if(bundle.containsKey("com.epicgames.ue4.GameActivity.bUseDisplayCutout"))
+			{
+				UseDisplayCutout = bundle.getBoolean("com.epicgames.ue4.GameActivity.bUseDisplayCutout");
+			}
 		}
 		catch (NameNotFoundException e)
 		{
@@ -48,6 +54,27 @@ public class SplashActivity extends Activity
 			}
 		}
 
+		// for now only allow on one device manufacturer - fix up later
+		if (!android.os.Build.MANUFACTURER.equals("HUAWEI"))
+		{
+			UseDisplayCutout = false;
+		}
+
+		if (UseDisplayCutout)
+		{
+			// only do this on Android Pie and above
+			if (android.os.Build.VERSION.SDK_INT >= 28)
+			{
+	            WindowManager.LayoutParams params = getWindow().getAttributes();
+		        params.layoutInDisplayCutoutMode = WindowManager.LayoutParams.LAYOUT_IN_DISPLAY_CUTOUT_MODE_SHORT_EDGES;
+			    getWindow().setAttributes(params);
+			}
+			else
+			{
+				UseDisplayCutout = false;
+			}
+		}
+
 		Intent intent = new Intent(this, GameActivity.class);
 		intent.putExtras(getIntent());
 		intent.addFlags(Intent.FLAG_ACTIVITY_NO_ANIMATION);
@@ -56,6 +83,10 @@ public class SplashActivity extends Activity
 		{
 			intent.putExtra("ShouldHideUI", "true");
 		}
+		if (UseDisplayCutout)
+		{
+			intent.putExtra("UseDisplayCutout", "true");
+		}
 
 		//pass down any extras added to this Activity's intent to the GameActivity intent (GCM data, for example)
 		Intent intentFromActivity = getIntent();
diff --git a/Engine/Config/BaseEngine.ini b/Engine/Config/BaseEngine.ini
--- a/Engine/Config/BaseEngine.ini
+++ b/Engine/Config/BaseEngine.ini
@@ -1894,6 +1894,7 @@ bBuildWithHiddenSymbolVisibility=false
 bSaveSymbols=false
 bAllowControllers=True
 bAllowIMU=True
+bUseDisplayCutout=False
 
 [/Script/AndroidPlatformEditor.AndroidSDKSettings]
 SDKAPILevel=latest
diff --git a/Engine/Source/Programs/UnrealBuildTool/Platform/Android/UEDeployAndroid.cs b/Engine/Source/Programs/UnrealBuildTool/Platform/Android/UEDeployAndroid.cs
--- a/Engine/Source/Programs/UnrealBuildTool/Platform/Android/UEDeployAndroid.cs
+++ b/Engine/Source/Programs/UnrealBuildTool/Platform/Android/UEDeployAndroid.cs
@@ -18,7 +18,7 @@ namespace UnrealBuildTool
 		private const string XML_HEADER = "<?xml version=\"1.0\" encoding=\"utf-8\"?>";
 
 		// Minimum Android SDK that must be used for Java compiling
-		readonly int MinimumSDKLevel = 23;
+		readonly int MinimumSDKLevel = 28;
 
 		// Minimum SDK version needed for Gradle based on active plugins (14 is for Google Play Services 11.0.4)
 		private int MinimumSDKLevelForGradle = 14;
@@ -2118,6 +2118,8 @@ namespace UnrealBuildTool
 			Ini.GetString("/Script/AndroidRuntimeSettings.AndroidRuntimeSettings", "Orientation", out Orientation);
 			bool EnableFullScreen;
 			Ini.GetBool("/Script/AndroidRuntimeSettings.AndroidRuntimeSettings", "bFullScreen", out EnableFullScreen);
+			bool bUseDisplayCutout;
+			Ini.GetBool("/Script/AndroidRuntimeSettings.AndroidRuntimeSettings", "bUseDisplayCutout", out bUseDisplayCutout);
 			List<string> ExtraManifestNodeTags;
 			Ini.GetArray("/Script/AndroidRuntimeSettings.AndroidRuntimeSettings", "ExtraManifestNodeTags", out ExtraManifestNodeTags);
 			List<string> ExtraApplicationNodeTags;
@@ -2421,6 +2423,7 @@ namespace UnrealBuildTool
 			Text.AppendLine(string.Format("\t\t<meta-data android:name=\"com.epicgames.ue4.GameActivity.CookedFlavors\" android:value=\"{0}\"/>", CookedFlavors));
 			Text.AppendLine(string.Format("\t\t<meta-data android:name=\"com.epicgames.ue4.GameActivity.bValidateTextureFormats\" android:value=\"{0}\"/>", bValidateTextureFormats ? "true" : "false"));
 			Text.AppendLine(string.Format("\t\t<meta-data android:name=\"com.epicgames.ue4.GameActivity.bUseExternalFilesDir\" android:value=\"{0}\"/>", bUseExternalFilesDir ? "true" : "false"));
+			Text.AppendLine(string.Format("\t\t<meta-data android:name=\"com.epicgames.ue4.GameActivity.bUseDisplayCutout\" android:value=\"{0}\"/>", bUseDisplayCutout ? "true" : "false"));
 			Text.AppendLine(string.Format("\t\t<meta-data android:name=\"com.epicgames.ue4.GameActivity.bAllowIMU\" android:value=\"{0}\"/>", bAllowIMU ? "true" : "false"));
 			Text.AppendLine(string.Format("\t\t<meta-data android:name=\"com.epicgames.ue4.GameActivity.bSupportsVulkan\" android:value=\"{0}\"/>", bSupportsVulkan ? "true" : "false"));
 			if (bUseNEONForArmV7)
diff --git a/Engine/Source/Runtime/Android/AndroidRuntimeSettings/Classes/AndroidRuntimeSettings.h b/Engine/Source/Runtime/Android/AndroidRuntimeSettings/Classes/AndroidRuntimeSettings.h
--- a/Engine/Source/Runtime/Android/AndroidRuntimeSettings/Classes/AndroidRuntimeSettings.h
+++ b/Engine/Source/Runtime/Android/AndroidRuntimeSettings/Classes/AndroidRuntimeSettings.h
@@ -245,6 +245,10 @@ public:
 	UPROPERTY(GlobalConfig, EditAnywhere, Category = "APK Packaging", Meta = (DisplayName = "Maximum supported aspect ratio."))
 	float MaxAspectRatio;
 
+	// Enables use of the display cutout area on Android 9+
+	UPROPERTY(GlobalConfig, EditAnywhere, Category = "APK Packaging", Meta = (DisplayName = "Use display cutout region?"))
+	bool bUseDisplayCutout;
+
 	// Level of verbosity to use during packaging with Ant
 	UPROPERTY(GlobalConfig, EditAnywhere, Category = "APK Packaging")
 	TEnumAsByte<EAndroidAntVerbosity::Type> AntVerbosity;
