Allow URLs other than http[s]:// on Android for LaunchURL

#ue4
#android

diff --git a/Engine/Build/Android/Java/src/com/epicgames/ue4/GameActivity.java.template b/Engine/Build/Android/Java/src/com/epicgames/ue4/GameActivity.java.template
--- a/Engine/Build/Android/Java/src/com/epicgames/ue4/GameActivity.java.template
+++ b/Engine/Build/Android/Java/src/com/epicgames/ue4/GameActivity.java.template
@@ -3675,8 +3675,12 @@ public class GameActivity extends NativeActivity implements SurfaceHolder.Callba
 		Log.debug("[JAVA} AndroidThunkJava_LaunchURL: URL = " + URL);
 		if (!URL.contains("://"))
 		{
-			URL = "http://" + URL;
-			Log.debug("[JAVA} AndroidThunkJava_LaunchURL: corrected URL = " + URL);
+			// add http:// if there isn't a scheme before a colon
+			if (URL.indexOf(":") < 1)
+			{
+				URL = "http://" + URL;
+				Log.debug("[JAVA} AndroidThunkJava_LaunchURL: corrected URL = " + URL);
+			}
 		}
 		try
 		{
