Put the notifcation body details as second line of notifcation on Android 5.0+ instead of ticker

#ue4
#android

diff --git a/Engine/Build/Android/Java/src/com/epicgames/ue4/LocalNotificationReceiver.java b/Engine/Build/Android/Java/src/com/epicgames/ue4/LocalNotificationReceiver.java
--- a/Engine/Build/Android/Java/src/com/epicgames/ue4/LocalNotificationReceiver.java
+++ b/Engine/Build/Android/Java/src/com/epicgames/ue4/LocalNotificationReceiver.java
@@ -51,6 +51,7 @@
 			.setContentTitle(title);
 		if (android.os.Build.VERSION.SDK_INT >= 21)
 		{
+			builder.setContentText(details);
 			builder.setColor(0xff0e1e43);
 		}
 		Notification notification = builder.build();
