Fix issue with URL based Android movie playback

#ue4
#android

diff --git a/Engine/Build/Android/Java/src/com/epicgames/ue4/MediaPlayer14.java b/Engine/Build/Android/Java/src/com/epicgames/ue4/MediaPlayer14.java
--- a/Engine/Build/Android/Java/src/com/epicgames/ue4/MediaPlayer14.java
+++ b/Engine/Build/Android/Java/src/com/epicgames/ue4/MediaPlayer14.java
@@ -269,10 +269,23 @@ public class MediaPlayer14
 				MediaExtractor extractor = new MediaExtractor();
 				if (extractor != null)
 				{
-					extractor.setDataSource(UrlPath);
-					updateTrackInfo(extractor);
-					extractor.release();
-					extractor = null;
+					try
+					{
+						extractor.setDataSource(UrlPath);
+						updateTrackInfo(extractor);
+						extractor.release();
+						extractor = null;
+					}
+					catch (Exception e)
+					{
+						GameActivity.Log.debug("setDataSourceURL: Exception = " + e);
+
+						// unable to collect track info, but can still try to play it
+						GameActivity.Log.debug("setDataSourceURL: Continuing without track info");
+						extractor.release();
+						extractor = null;
+						return true;
+					}
 				}
 			}
 		}
