Skip track info extraction if startup movie on Android

#ue4
#android

diff --git a/Engine/Build/Android/Java/src/com/epicgames/ue4/MediaPlayer14.java b/Engine/Build/Android/Java/src/com/epicgames/ue4/MediaPlayer14.java
--- a/Engine/Build/Android/Java/src/com/epicgames/ue4/MediaPlayer14.java
+++ b/Engine/Build/Android/Java/src/com/epicgames/ue4/MediaPlayer14.java
@@ -30,6 +30,7 @@ public class MediaPlayer14
 {
 	private boolean SwizzlePixels = true;
 	private boolean VulkanRenderer = false;
+	private boolean NeedTrackInfo = true;
 	private boolean Looping = false;
 	private boolean AudioEnabled = true;
 	private float AudioVolume = 1.0f;
@@ -82,10 +83,11 @@ public class MediaPlayer14
 
 	// ======================================================================================
 
-	public MediaPlayer14(boolean swizzlePixels, boolean vulkanRenderer)
+	public MediaPlayer14(boolean swizzlePixels, boolean vulkanRenderer, boolean needTrackInfo)
 	{
 		SwizzlePixels = swizzlePixels;
 		VulkanRenderer = vulkanRenderer;
+		NeedTrackInfo = needTrackInfo;
 		WaitOnBitmapRender = false;
 		AudioEnabled = true;
 		AudioVolume = 1.0f;
@@ -265,7 +267,7 @@ public class MediaPlayer14
 			setDataSource(UrlPath);
 			releaseOESTextureRenderer();
 			releaseBitmapRenderer();
-			if (android.os.Build.VERSION.SDK_INT >= 16)
+			if (NeedTrackInfo && android.os.Build.VERSION.SDK_INT >= 16)
 			{
 				MediaExtractor extractor = new MediaExtractor();
 				if (extractor != null)
@@ -384,7 +386,7 @@ public class MediaPlayer14
 			releaseOESTextureRenderer();
 			releaseBitmapRenderer();
 
-			if (android.os.Build.VERSION.SDK_INT >= 16)
+			if (NeedTrackInfo && android.os.Build.VERSION.SDK_INT >= 16)
 			{
 				MediaExtractor extractor = new MediaExtractor();
 				if (extractor != null)
@@ -432,7 +434,7 @@ public class MediaPlayer14
 			releaseOESTextureRenderer();
 			releaseBitmapRenderer();
 
-			if (android.os.Build.VERSION.SDK_INT >= 16)
+			if (NeedTrackInfo && android.os.Build.VERSION.SDK_INT >= 16)
 			{
 				MediaExtractor extractor = new MediaExtractor();
 				if (extractor != null)
@@ -474,7 +476,7 @@ public class MediaPlayer14
 			releaseOESTextureRenderer();
 			releaseBitmapRenderer();
 
-			if (android.os.Build.VERSION.SDK_INT >= 16)
+			if (NeedTrackInfo && android.os.Build.VERSION.SDK_INT >= 16)
 			{
 				MediaExtractor extractor = new MediaExtractor();
 				if (extractor != null)
@@ -1846,7 +1848,7 @@ public class MediaPlayer14
 
 	public AudioTrackInfo[] GetAudioTracks()
 	{
-		if (android.os.Build.VERSION.SDK_INT >= 16)
+		if (NeedTrackInfo && android.os.Build.VERSION.SDK_INT >= 16)
 		{
 			TrackInfo[] trackInfo = getTrackInfo();
 			int CountTracks = 0;
@@ -1922,7 +1924,7 @@ public class MediaPlayer14
 
 	public CaptionTrackInfo[] GetCaptionTracks()
 	{
-		if (android.os.Build.VERSION.SDK_INT >= 21)
+		if (NeedTrackInfo && android.os.Build.VERSION.SDK_INT >= 21)
 		{
 			TrackInfo[] trackInfo = getTrackInfo();
 			int CountTracks = 0;
@@ -1974,7 +1976,7 @@ public class MediaPlayer14
 		int Width = getVideoWidth();
 		int Height = getVideoHeight();
 
-		if (android.os.Build.VERSION.SDK_INT >= 16)
+		if (NeedTrackInfo && android.os.Build.VERSION.SDK_INT >= 16)
 		{
 			TrackInfo[] trackInfo = getTrackInfo();
 			int CountTracks = 0;
diff --git a/Engine/Plugins/Media/AndroidMedia/Source/AndroidMedia/Private/Player/AndroidMediaPlayer.cpp b/Engine/Plugins/Media/AndroidMedia/Source/AndroidMedia/Private/Player/AndroidMediaPlayer.cpp
--- a/Engine/Plugins/Media/AndroidMedia/Source/AndroidMedia/Private/Player/AndroidMediaPlayer.cpp
+++ b/Engine/Plugins/Media/AndroidMedia/Source/AndroidMedia/Private/Player/AndroidMediaPlayer.cpp
@@ -38,9 +38,9 @@ FAndroidMediaPlayer::FAndroidMediaPlayer(IMediaEventSink& InEventSink)
 	, bLooping(false)
 	, EventSink(InEventSink)
 #if WITH_ENGINE
-	, JavaMediaPlayer(MakeShared<FJavaAndroidMediaPlayer, ESPMode::ThreadSafe>(false, FAndroidMisc::ShouldUseVulkan()))
+	, JavaMediaPlayer(MakeShared<FJavaAndroidMediaPlayer, ESPMode::ThreadSafe>(false, FAndroidMisc::ShouldUseVulkan(), true))
 #else
-	, JavaMediaPlayer(MakeShared<FJavaAndroidMediaPlayer, ESPMode::ThreadSafe>(true, FAndroidMisc::ShouldUseVulkan()))
+	, JavaMediaPlayer(MakeShared<FJavaAndroidMediaPlayer, ESPMode::ThreadSafe>(true, FAndroidMisc::ShouldUseVulkan(), true))
 #endif
 	, Samples(MakeShared<FMediaSamples, ESPMode::ThreadSafe>())
 	, SelectedAudioTrack(INDEX_NONE)
diff --git a/Engine/Plugins/Runtime/AndroidMoviePlayer/Source/AndroidMoviePlayer/Private/AndroidMovieStreamer.cpp b/Engine/Plugins/Runtime/AndroidMoviePlayer/Source/AndroidMoviePlayer/Private/AndroidMovieStreamer.cpp
--- a/Engine/Plugins/Runtime/AndroidMoviePlayer/Source/AndroidMoviePlayer/Private/AndroidMovieStreamer.cpp
+++ b/Engine/Plugins/Runtime/AndroidMoviePlayer/Source/AndroidMoviePlayer/Private/AndroidMovieStreamer.cpp
@@ -27,7 +27,7 @@ FAndroidMediaPlayerStreamer::FAndroidMediaPlayerStreamer()
 	: JavaMediaPlayer(nullptr)
 	, CurrentPosition(-1)
 {
-	JavaMediaPlayer = MakeShareable(new FJavaAndroidMediaPlayer(false, FAndroidMisc::ShouldUseVulkan()));
+	JavaMediaPlayer = MakeShareable(new FJavaAndroidMediaPlayer(false, FAndroidMisc::ShouldUseVulkan(), false));
 	MovieViewport = MakeShareable(new FMovieViewport());
 }
 
diff --git a/Engine/Source/Runtime/ApplicationCore/Private/Android/AndroidJavaMediaPlayer.cpp b/Engine/Source/Runtime/ApplicationCore/Private/Android/AndroidJavaMediaPlayer.cpp
--- a/Engine/Source/Runtime/ApplicationCore/Private/Android/AndroidJavaMediaPlayer.cpp
+++ b/Engine/Source/Runtime/ApplicationCore/Private/Android/AndroidJavaMediaPlayer.cpp
@@ -119,8 +119,8 @@ int32 FJavaAndroidMediaDataSource::ReadAt(int64 Position, uint8* Buffer, int32 C
 /* JavaAndroidMediaPlayer
 *****************************************************************************/
 
-FJavaAndroidMediaPlayer::FJavaAndroidMediaPlayer(bool swizzlePixels, bool vulkanRenderer)
-	: FJavaClassObject(GetClassName(), "(ZZ)V", swizzlePixels, vulkanRenderer)
+FJavaAndroidMediaPlayer::FJavaAndroidMediaPlayer(bool swizzlePixels, bool vulkanRenderer, bool needTrackInfo)
+	: FJavaClassObject(GetClassName(), "(ZZZ)V", swizzlePixels, vulkanRenderer, needTrackInfo)
 	, GetDurationMethod(GetClassMethod("getDuration", "()I"))
 	, ResetMethod(GetClassMethod("reset", "()V"))
 	, GetCurrentPositionMethod(GetClassMethod("getCurrentPosition", "()I"))
diff --git a/Engine/Source/Runtime/ApplicationCore/Public/Android/AndroidJavaMediaPlayer.h b/Engine/Source/Runtime/ApplicationCore/Public/Android/AndroidJavaMediaPlayer.h
--- a/Engine/Source/Runtime/ApplicationCore/Public/Android/AndroidJavaMediaPlayer.h
+++ b/Engine/Source/Runtime/ApplicationCore/Public/Android/AndroidJavaMediaPlayer.h
@@ -67,7 +67,7 @@ public:
 	};
 
 public:
-	FJavaAndroidMediaPlayer(bool swizzlePixels, bool vulkanRenderer);
+	FJavaAndroidMediaPlayer(bool swizzlePixels, bool vulkanRenderer, bool needTrackInfo);
 	int32 GetDuration();
 	void Reset();
 	void Stop();
