Updating AndroidThunkJava_IsGamepadAttached to not check SOURCE_DPAD because it was false flagging gamepads.

Chris.Babcock

diff --git a/Engine/Build/Android/Java/src/com/epicgames/ue4/GameActivity.java.template b/Engine/Build/Android/Java/src/com/epicgames/ue4/GameActivity.java.template
--- a/Engine/Build/Android/Java/src/com/epicgames/ue4/GameActivity.java.template
+++ b/Engine/Build/Android/Java/src/com/epicgames/ue4/GameActivity.java.template
@@ -4696,8 +4696,7 @@ public class GameActivity extends $${gameActivitySuperClass}$$ implements Surfac
 			InputDevice inputDevice = InputDevice.getDevice(deviceIds[deviceIndex]);
 			// is it a joystick, dpad, or gamepad?
 			if (((inputDevice.getSources() & InputDevice.SOURCE_GAMEPAD) == InputDevice.SOURCE_GAMEPAD)
-                || ((inputDevice.getSources() & InputDevice.SOURCE_JOYSTICK) == InputDevice.SOURCE_JOYSTICK)
-				|| ((inputDevice.getSources() & InputDevice.SOURCE_DPAD) == InputDevice.SOURCE_DPAD))
+                || ((inputDevice.getSources() & InputDevice.SOURCE_JOYSTICK) == InputDevice.SOURCE_JOYSTICK))
 			{
 				return true;
 			}
