Merging 5360223 from Dev-EngineMerge. DevK integration.

diff --git a/Engine/Build/Android/Java/src/com/epicgames/ue4/GameActivity.java.template b/Engine/Build/Android/Java/src/com/epicgames/ue4/GameActivity.java.template
--- a/Engine/Build/Android/Java/src/com/epicgames/ue4/GameActivity.java.template
+++ b/Engine/Build/Android/Java/src/com/epicgames/ue4/GameActivity.java.template
@@ -39,6 +39,7 @@ import java.util.ArrayList;
 import java.text.DecimalFormat;
 import java.text.DecimalFormatSymbols;
 import java.util.zip.Inflater;
+import java.util.Stack;
 import java.util.Timer;
 import java.util.TimerTask;
 import java.lang.Runnable;
@@ -277,6 +278,12 @@ public class GameActivity extends $${gameActivitySuperClass}$$ implements Surfac
 	public static final int DOWNLOAD_NO_PLAY_KEY = 6;
 	public static final String DOWNLOAD_RETURN_NAME = "Result";
 
+	private static final int CONFIGRULES_STATE_RUN = 0;
+	private static final int CONFIGRULES_STATE_EXECTRUE = 1;
+	private static final int CONFIGRULES_STATE_FINDELSE = 2;
+	private static final int CONFIGRULES_STATE_EXECFALSE = 3;
+	private static final int CONFIGRULES_STATE_FINDEND = 4;
+
 	public static enum VirtualKeyboardCommand { VK_CMD_NONE, VK_CMD_SHOW, VK_CMD_HIDE };
 	public static VirtualKeyboardCommand lastVirtualKeyboardCommand = VirtualKeyboardCommand.VK_CMD_NONE;
 	public static final int lastVirtualKeyboardCommandDelay = 200;
@@ -1251,6 +1258,222 @@ public class GameActivity extends $${gameActivitySuperClass}$$ implements Surfac
 		return Result;
 	}
 
+	private boolean EvaluateConditions(Map<String, String> variables, ArrayList<String> conditions, String PreviousRegexMatch)
+	{
+		boolean bConditionTrue = true;
+		for (String condition : conditions)
+		{
+			String SourceType = "";
+			String CompareType = "";
+			String MatchString = "";
+
+			// deal with condition group (src,cmp,match)
+			ArrayList<String> groups = ParseSegments(RemoveSurrounds(condition, "(", ")"), ",", "\"", "\"");
+			for (String group : groups)
+			{
+				ArrayList<String> keyvalue = ParseSegments(group, "=", "\"", "\"");
+				if (keyvalue.size() == 2)
+				{
+					String key = RemoveSurrounds(keyvalue.get(0), "\"", "\"");
+					String value = RemoveSurrounds(keyvalue.get(1), "\"", "\"");
+
+					if (key.equals("SourceType"))
+					{
+						SourceType = value;
+					}
+					else if (key.equals("CompareType"))
+					{
+						CompareType = value;
+					}
+					else if (key.equals("MatchString"))
+					{
+						MatchString = value;
+					}
+				}
+			}
+
+			String Source = "";
+			boolean bSourceIsCommandLine = false;
+			if (SourceType.equals("SRC_PreviousRegexMatch"))
+			{
+				Source = PreviousRegexMatch;
+			}
+			else if (SourceType.equals("SRC_CommandLine"))
+			{
+				bSourceIsCommandLine = true;
+				Source = CommandLineFull;;
+			}
+			else if (variables.containsKey(SourceType))
+			{
+				Source = variables.get(SourceType);
+			}
+			else if (SourceType.equals("[EXIST]"))
+			{
+				Source = MatchString;
+			}
+			else
+			{
+				bConditionTrue = false;
+				break;
+			}
+
+			// apply operation
+			if (CompareType.equals("CMP_Exist"))
+			{
+				if (!variables.containsKey(Source))
+				{
+					bConditionTrue = false;
+					break;
+				}
+			}
+			else if (CompareType.equals("CMP_NotExist"))
+			{
+				if (variables.containsKey(Source))
+				{
+					bConditionTrue = false;
+					break;
+				}
+			}
+			else if (CompareType.equals("CMP_Equal"))
+			{
+				if (bSourceIsCommandLine)
+				{
+					if (!CommandLineFlags.containsKey(MatchString))
+					{
+						bConditionTrue = false;
+						break;
+					}
+				}
+				else if (!Source.equals(MatchString))
+				{
+					bConditionTrue = false;
+					break;
+				}
+			}
+			else if (CompareType.equals("CMP_NotEqual"))
+			{
+				if (bSourceIsCommandLine)
+				{
+					if (CommandLineFlags.containsKey(MatchString))
+					{
+						bConditionTrue = false;
+						break;
+					}
+				}
+				else if (Source.equals(MatchString))
+				{
+					bConditionTrue = false;
+					break;
+				}
+			}
+			else if (CompareType.equals("CMP_EqualIgnore"))
+			{
+				if (!Source.toLowerCase().equals(MatchString.toLowerCase()))
+				{
+					bConditionTrue = false;
+					break;
+				}
+			}
+			else if (CompareType.equals("CMP_NotEqualIgnore"))
+			{
+				if (Source.toLowerCase().equals(MatchString.toLowerCase()))
+				{
+					bConditionTrue = false;
+					break;
+				}
+			}
+			else if (CompareType.equals("CMP_Regex"))
+			{
+				Pattern pattern = Pattern.compile(MatchString);
+				Matcher matcher = pattern.matcher(Source);
+				if (matcher.find())
+				{
+					int groupCount = matcher.groupCount();
+					PreviousRegexMatch = matcher.group(groupCount >= 1 ? 1 : 0);
+				}
+				else
+				{
+					bConditionTrue = false;
+					break;
+				}
+			}
+			else
+			{
+				boolean bNumericOperands = true;
+				float SourceFloat = 0.0f;
+				float MatchFloat = 0.0f;
+
+				// convert source and match to float if numeric
+				try
+				{
+					SourceFloat = Float.parseFloat(Source);
+				}
+				catch (NumberFormatException ne)
+				{
+					bNumericOperands = false;
+				}
+				try
+				{
+					MatchFloat = Float.parseFloat(MatchString);
+				}
+				catch (NumberFormatException ne)
+				{
+					bNumericOperands = false;
+				}
+
+				// if comparison ends with Ignore, do case-insensitive compare by converting both to lowercase
+				if (CompareType.endsWith("Ignore"))
+				{
+					bNumericOperands = false;
+					CompareType = CompareType.substring(0, CompareType.length() - 6);
+
+					Source = Source.toLowerCase();
+					MatchString = MatchString.toLowerCase();
+				}
+
+				if (CompareType.equals("CMP_Less"))
+				{
+					if ((bNumericOperands && (SourceFloat >= MatchFloat)) || (!bNumericOperands && (Source.compareTo(MatchString) >= 0)))
+					{
+						bConditionTrue = false;
+						break;
+					}
+				}
+				else if (CompareType.equals("CMP_LessEqual"))
+				{
+					if ((bNumericOperands && (SourceFloat > MatchFloat)) || (!bNumericOperands && (Source.compareTo(MatchString) > 0)))
+					{
+						bConditionTrue = false;
+						break;
+					}
+				}
+				else if (CompareType.equals("CMP_Greater"))
+				{
+					if ((bNumericOperands && (SourceFloat <= MatchFloat)) || (!bNumericOperands && (Source.compareTo(MatchString) <= 0)))
+					{
+						bConditionTrue = false;
+						break;
+					}
+				}
+				else if (CompareType.equals("CMP_GreaterEqual"))
+				{
+					if ((bNumericOperands && (SourceFloat <= MatchFloat)) || (!bNumericOperands && (Source.compareTo(MatchString) < 0)))
+					{
+						bConditionTrue = false;
+						break;
+					}
+				}
+				else
+				{
+					bConditionTrue = false;
+					break;
+				}
+			}
+		}
+
+		return bConditionTrue;
+	}
+
 	private boolean processSystemInfo(String ProjectName, String PackageName)
 	{
 		// check OpenGL information
@@ -1395,6 +1618,8 @@ public class GameActivity extends $${gameActivitySuperClass}$$ implements Surfac
 		// commandline flag may be used to skip dialog if not a shipping build
 		boolean bIgnoreDialog = !nativeIsShippingBuild() && CommandLineFlags.containsKey("noconfigrulesdialog");
 
+//$${gameActivityPreConfigRulesParseAdditions}$$
+		
 		// read the configuration file (configrules.bin.png)
 		int configRulesVersion = 0;
 		BufferedReader reader = null;
@@ -1419,6 +1644,10 @@ public class GameActivity extends $${gameActivitySuperClass}$$ implements Surfac
 			
 			try
 			{
+				Stack<Integer> StateStack = new Stack<Integer>();
+
+				int currentState = CONFIGRULES_STATE_RUN;
+				int nestDepth = 0;
 				String line;
 				while ((line = reader.readLine()) != null)
 				{
@@ -1445,6 +1674,105 @@ public class GameActivity extends $${gameActivitySuperClass}$$ implements Surfac
 					String command = line.substring(0, index).trim();
 					line = line.substring(index + 1).trim();
 
+					// handle states
+					switch (currentState)
+					{
+						case CONFIGRULES_STATE_RUN:
+							if (command.equals("else") || command.equals("elseif") || command.equals("endif"))
+							{
+								Log.debug("ConfigRules: unexpected " + command + " encountered!");
+							}
+							break;
+
+						case CONFIGRULES_STATE_EXECTRUE:
+							if (command.equals("else"))
+							{
+								currentState = CONFIGRULES_STATE_FINDEND;
+								continue;
+							}
+							if (command.equals("endif"))
+							{
+								currentState = StateStack.pop();
+								continue;
+							}
+							break;
+
+						case CONFIGRULES_STATE_FINDELSE:
+							if (command.equals("if"))
+							{
+								nestDepth++;
+								continue;
+							}
+							if (nestDepth > 0)
+							{
+								if (command.equals("endif"))
+								{
+									nestDepth--;
+								}
+								continue;
+							}
+							if (command.equals("endif"))
+							{
+								currentState = StateStack.pop();
+								continue;
+							}
+							if (command.equals("else"))
+							{
+								currentState = CONFIGRULES_STATE_EXECFALSE;
+								continue;
+							}
+							if (command.equals("elseif"))
+							{
+								currentState = CONFIGRULES_STATE_FINDEND;
+
+								ArrayList<String> conditions = ParseSegments(line, ",", "(\"", ")\"");
+								int setsize = conditions.size();
+								if (setsize > 0)
+								{
+									boolean bConditionTrue = EvaluateConditions(variables, conditions, "");
+									currentState = bConditionTrue ? CONFIGRULES_STATE_EXECTRUE : CONFIGRULES_STATE_FINDELSE;
+								}
+							}
+							continue;
+
+						case CONFIGRULES_STATE_EXECFALSE:
+							if (command.equals("endif"))
+							{
+								currentState = StateStack.pop();
+								continue;
+							}
+							if (command.equals("else") || command.equals("elseif"))
+							{
+								Log.debug("ConfigRules: unexpected " + command + " while handling false condition!");
+								continue;
+							}
+							break;
+
+						case CONFIGRULES_STATE_FINDEND:
+							if (command.equals("if"))
+							{
+								nestDepth++;
+								continue;
+							}
+							if (nestDepth > 0)
+							{
+								if (command.equals("endif"))
+								{
+									nestDepth--;
+								}
+								continue;
+							}
+							if (command.equals("endif"))
+							{
+								currentState = StateStack.pop();
+							}
+							continue;
+
+						default:
+							Log.debug("ConfigRules: unknown state!");
+							continue;
+					}
+
 					// handle commands
 					if (command.equals("set"))
 					{
@@ -1508,6 +1836,26 @@ public class GameActivity extends $${gameActivitySuperClass}$$ implements Surfac
 							}
 						}
 					}
+					else if (command.equals("if"))
+					{
+						// if:(SourceType=SRC_DeviceMake,CompareType=CMP_Equal,MatchString="samsung")
+						// ... commands for true for all conditions
+						// elseif:(SourceType=SRC_DeviceMake,CompareType=CMP_Equal,MatchString="Google")
+						// ... commands for true for all conditions
+						// else:
+						// ... commands for false for any condition
+						// end:
+						StateStack.push(currentState);
+						currentState = CONFIGRULES_STATE_FINDEND;
+
+						ArrayList<String> conditions = ParseSegments(line, ",", "(\"", ")\"");
+						int setsize = conditions.size();
+						if (setsize > 0)
+						{
+							boolean bConditionTrue = EvaluateConditions(variables, conditions, "");
+							currentState = bConditionTrue ? CONFIGRULES_STATE_EXECTRUE : CONFIGRULES_STATE_FINDELSE;
+						}
+					}
 					else if (command.equals("condition"))
 					{
 						// condition:((SourceType=SRC_DeviceMake,CompareType=CMP_Equal,MatchString="samsung")),(SourceType=,CompareType=,MatchString=),...]),(a=b[,c=d,...]),(a[,b,...])
@@ -1520,217 +1868,7 @@ public class GameActivity extends $${gameActivitySuperClass}$$ implements Surfac
 							ArrayList<String> sets = ParseSegments(RemoveSurrounds(conditionAndSets.get(1), "(", ")"), ",", "(\"", ")\"");
 							ArrayList<String> clears = (setsize == 3) ? ParseSegments(RemoveSurrounds(conditionAndSets.get(2), "(", ")"), ",", "(\"", ")\"") : new ArrayList<String>();
 
-							String PreviousRegexMatch = "";
-							boolean bConditionTrue = true;
-							for (String condition : conditions)
-							{
-								String SourceType = "";
-								String CompareType = "";
-								String MatchString = "";
-
-								// deal with condition group (src,cmp,match)
-								ArrayList<String> groups = ParseSegments(RemoveSurrounds(condition, "(", ")"), ",", "\"", "\"");
-								for (String group : groups)
-								{
-									ArrayList<String> keyvalue = ParseSegments(group, "=", "\"", "\"");
-									if (keyvalue.size() == 2)
-									{
-										String key = RemoveSurrounds(keyvalue.get(0), "\"", "\"");
-										String value = RemoveSurrounds(keyvalue.get(1), "\"", "\"");
-
-										if (key.equals("SourceType"))
-										{
-											SourceType = value;
-										}
-										else if (key.equals("CompareType"))
-										{
-											CompareType = value;
-										}
-										else if (key.equals("MatchString"))
-										{
-											MatchString = value;
-										}
-									}
-								}
-
-								String Source = "";
-								boolean bSourceIsCommandLine = false;
-								if (SourceType.equals("SRC_PreviousRegexMatch"))
-								{
-									Source = PreviousRegexMatch;
-								}
-								else if (SourceType.equals("SRC_CommandLine"))
-								{
-									bSourceIsCommandLine = true;
-									Source = CommandLineFull;;
-								}
-								else if (variables.containsKey(SourceType))
-								{
-									Source = variables.get(SourceType);
-								}
-								else if (SourceType.equals("[EXIST]"))
-								{
-									Source = MatchString;
-								}
-								else
-								{
-									bConditionTrue = false;
-									break;
-								}
-
-								// apply operation
-								if (CompareType.equals("CMP_Exist"))
-								{
-									if (!variables.containsKey(Source))
-									{
-										bConditionTrue = false;
-										break;
-									}
-								}
-								else if (CompareType.equals("CMP_NotExist"))
-								{
-									if (variables.containsKey(Source))
-									{
-										bConditionTrue = false;
-										break;
-									}
-								}
-								else if (CompareType.equals("CMP_Equal"))
-								{
-									if (bSourceIsCommandLine)
-									{
-										if (!CommandLineFlags.containsKey(MatchString))
-										{
-											bConditionTrue = false;
-											break;
-										}
-									}
-									else if (!Source.equals(MatchString))
-									{
-										bConditionTrue = false;
-										break;
-									}
-								}
-								else if (CompareType.equals("CMP_NotEqual"))
-								{
-									if (bSourceIsCommandLine)
-									{
-										if (CommandLineFlags.containsKey(MatchString))
-										{
-											bConditionTrue = false;
-											break;
-										}
-									}
-									else if (Source.equals(MatchString))
-									{
-										bConditionTrue = false;
-										break;
-									}
-								}
-								else if (CompareType.equals("CMP_EqualIgnore"))
-								{
-									if (!Source.toLowerCase().equals(MatchString.toLowerCase()))
-									{
-										bConditionTrue = false;
-										break;
-									}
-								}
-								else if (CompareType.equals("CMP_NotEqualIgnore"))
-								{
-									if (Source.toLowerCase().equals(MatchString.toLowerCase()))
-									{
-										bConditionTrue = false;
-										break;
-									}
-								}
-								else if (CompareType.equals("CMP_Regex"))
-								{
-									Pattern pattern = Pattern.compile(MatchString);
-									Matcher matcher = pattern.matcher(Source);
-									if (matcher.find())
-									{
-										int groupCount = matcher.groupCount();
-										PreviousRegexMatch = matcher.group(groupCount >= 1 ? 1 : 0);
-									}
-									else
-									{
-										bConditionTrue = false;
-										break;
-									}
-								}
-								else
-								{
-									boolean bNumericOperands = true;
-									float SourceFloat = 0.0f;
-									float MatchFloat = 0.0f;
-
-									// convert source and match to float if numeric
-									try
-									{
-										SourceFloat = Float.parseFloat(Source);
-									}
-									catch (NumberFormatException ne)
-									{
-										bNumericOperands = false;
-									}
-									try
-									{
-										MatchFloat = Float.parseFloat(MatchString);
-									}
-									catch (NumberFormatException ne)
-									{
-										bNumericOperands = false;
-									}
-
-									// if comparison ends with Ignore, do case-insensitive compare by converting both to lowercase
-									if (CompareType.endsWith("Ignore"))
-									{
-										bNumericOperands = false;
-										CompareType = CompareType.substring(0, CompareType.length() - 6);
-
-										Source = Source.toLowerCase();
-										MatchString = MatchString.toLowerCase();
-									}
-
-									if (CompareType.equals("CMP_Less"))
-									{
-										if ((bNumericOperands && (SourceFloat >= MatchFloat)) || (!bNumericOperands && (Source.compareTo(MatchString) >= 0)))
-										{
-											bConditionTrue = false;
-											break;
-										}
-									}
-									else if (CompareType.equals("CMP_LessEqual"))
-									{
-										if ((bNumericOperands && (SourceFloat > MatchFloat)) || (!bNumericOperands && (Source.compareTo(MatchString) > 0)))
-										{
-											bConditionTrue = false;
-											break;
-										}
-									}
-									else if (CompareType.equals("CMP_Greater"))
-									{
-										if ((bNumericOperands && (SourceFloat <= MatchFloat)) || (!bNumericOperands && (Source.compareTo(MatchString) <= 0)))
-										{
-											bConditionTrue = false;
-											break;
-										}
-									}
-									else if (CompareType.equals("CMP_GreaterEqual"))
-									{
-										if ((bNumericOperands && (SourceFloat <= MatchFloat)) || (!bNumericOperands && (Source.compareTo(MatchString) < 0)))
-										{
-											bConditionTrue = false;
-											break;
-										}
-									}
-									else
-									{
-										bConditionTrue = false;
-										break;
-									}
-								}
-							}
+							boolean bConditionTrue = EvaluateConditions(variables, conditions, "");
 
 							if (bConditionTrue)
 							{
@@ -1936,6 +2074,7 @@ public class GameActivity extends $${gameActivitySuperClass}$$ implements Surfac
 		nativeSetConfigRulesVariables(keyValues);
 
 //$${gameActivityPostConfigRulesAdditions}$$
+//$${gameActivityFinalizeConfigRulesAdditions}$$
 
 		return !bForceGameEndWithError;
 	}
diff --git a/Engine/Config/BaseEngine.ini b/Engine/Config/BaseEngine.ini
--- a/Engine/Config/BaseEngine.ini
+++ b/Engine/Config/BaseEngine.ini
@@ -12,7 +12,7 @@ GlobalDefaultGameMode="/Script/Engine.GameModeBase"
 LocalMapOptions=
 
 [InstallBundleManager]
-ModuleName=NullInstallBundleManagerModule
+ModuleName=NullInstallBundleManager
 
 [URL]
 Protocol=unreal
diff --git a/Engine/Plugins/Animation/LiveLink/Source/LiveLink/Private/AnimNode_LiveLinkArchiveComponentPose.cpp b/Engine/Plugins/Animation/LiveLink/Source/LiveLink/Private/AnimNode_LiveLinkArchiveComponentPose.cpp
deleted file mode 100644
diff --git a/Engine/Plugins/Animation/LiveLink/Source/LiveLink/Private/LiveLinkArchiveComponent.cpp b/Engine/Plugins/Animation/LiveLink/Source/LiveLink/Private/LiveLinkArchiveComponent.cpp
deleted file mode 100644
diff --git a/Engine/Plugins/Animation/LiveLink/Source/LiveLink/Public/AnimNode_LiveLinkArchiveComponentPose.h b/Engine/Plugins/Animation/LiveLink/Source/LiveLink/Public/AnimNode_LiveLinkArchiveComponentPose.h
deleted file mode 100644
diff --git a/Engine/Plugins/Animation/LiveLink/Source/LiveLink/Public/LiveLinkArchiveComponent.h b/Engine/Plugins/Animation/LiveLink/Source/LiveLink/Public/LiveLinkArchiveComponent.h
deleted file mode 100644
diff --git a/Engine/Plugins/Animation/LiveLink/Source/LiveLinkEditor/Private/AnimGraphNode_LiveLinkArchiveComponentPose.cpp b/Engine/Plugins/Animation/LiveLink/Source/LiveLinkEditor/Private/AnimGraphNode_LiveLinkArchiveComponentPose.cpp
deleted file mode 100644
diff --git a/Engine/Plugins/Animation/LiveLink/Source/LiveLinkEditor/Private/AnimGraphNode_LiveLinkArchiveComponentPose.h b/Engine/Plugins/Animation/LiveLink/Source/LiveLinkEditor/Private/AnimGraphNode_LiveLinkArchiveComponentPose.h
deleted file mode 100644
diff --git a/Engine/Plugins/Runtime/AR/Google/GoogleARCore/Source/GoogleARCoreBase/Private/GoogleARCoreXRTrackingSystem.cpp b/Engine/Plugins/Runtime/AR/Google/GoogleARCore/Source/GoogleARCoreBase/Private/GoogleARCoreXRTrackingSystem.cpp
--- a/Engine/Plugins/Runtime/AR/Google/GoogleARCore/Source/GoogleARCoreBase/Private/GoogleARCoreXRTrackingSystem.cpp
+++ b/Engine/Plugins/Runtime/AR/Google/GoogleARCore/Source/GoogleARCoreBase/Private/GoogleARCoreXRTrackingSystem.cpp
@@ -52,6 +52,11 @@ bool FGoogleARCoreXRTrackingSystem::IsHeadTrackingAllowed() const
 
 bool FGoogleARCoreXRTrackingSystem::GetCurrentPose(int32 DeviceId, FQuat& OutOrientation, FVector& OutPosition)
 {
+	if (OnGetARSessionStatus().Status != EARSessionStatus::Running)
+	{
+		return false;
+	}
+	
 	if (DeviceId == IXRTrackingSystem::HMDDeviceId)
 	{
 		OutOrientation = CachedOrientation;
@@ -194,6 +199,17 @@ EARTrackingQuality FGoogleARCoreXRTrackingSystem::OnGetTrackingQuality() const
 	return EARTrackingQuality::OrientationAndPosition;
 }
 
+EARTrackingQualityReason FGoogleARCoreXRTrackingSystem::OnGetTrackingQualityReason() const
+{
+	// Dont return EARTrackingQualityReason::None, which means that the tracking quality is not limited
+	return EARTrackingQualityReason::InsufficientFeatures;
+}
+
+bool FGoogleARCoreXRTrackingSystem::IsARAvailable() const
+{
+	return FGoogleARCoreDevice::GetInstance()->CheckARCoreAPKAvailability() == EGoogleARCoreAvailability::SupportedInstalled;
+}
+
 void FGoogleARCoreXRTrackingSystem::OnStartARSession(UARSessionConfig* SessionConfig)
 {
 	FGoogleARCoreDevice::GetInstance()->StartARCoreSessionRequest(SessionConfig);
diff --git a/Engine/Plugins/Runtime/AR/Google/GoogleARCore/Source/GoogleARCoreBase/Private/GoogleARCoreXRTrackingSystem.h b/Engine/Plugins/Runtime/AR/Google/GoogleARCore/Source/GoogleARCoreBase/Private/GoogleARCoreXRTrackingSystem.h
--- a/Engine/Plugins/Runtime/AR/Google/GoogleARCore/Source/GoogleARCoreBase/Private/GoogleARCoreXRTrackingSystem.h
+++ b/Engine/Plugins/Runtime/AR/Google/GoogleARCore/Source/GoogleARCoreBase/Private/GoogleARCoreXRTrackingSystem.h
@@ -48,8 +48,11 @@ public:
 
 protected:
 	// IARSystemSupport
+	/** Returns true/false based on whether AR features are available */
+	virtual bool IsARAvailable() const override;
 	virtual void OnARSystemInitialized() override;
 	virtual EARTrackingQuality OnGetTrackingQuality() const override;
+	virtual EARTrackingQualityReason OnGetTrackingQualityReason() const override;
 	virtual void OnStartARSession(UARSessionConfig* SessionConfig) override;
 	virtual void OnPauseARSession() override;
 	virtual void OnStopARSession() override;
diff --git a/Engine/Plugins/Runtime/Firebase/Source/Java/notifications/EpicFirebaseInstanceIDService.java b/Engine/Plugins/Runtime/Firebase/Source/Java/notifications/EpicFirebaseInstanceIDService.java
--- a/Engine/Plugins/Runtime/Firebase/Source/Java/notifications/EpicFirebaseInstanceIDService.java
+++ b/Engine/Plugins/Runtime/Firebase/Source/Java/notifications/EpicFirebaseInstanceIDService.java
@@ -22,7 +22,7 @@ public class EpicFirebaseInstanceIDService extends FirebaseInstanceIdService {
 
 	@Override
 	public void onTokenRefresh() {
-		String firebaseToken = getFirebaseInstanceId().getToken();
+		String firebaseToken = getFirebaseInstanceToken();
 		Log.debug("Refreshed Firebase token: " + firebaseToken);
 		if (TextUtils.isEmpty(firebaseToken)) {
 			Log.error("Firebase token is empty or null");
@@ -31,9 +31,19 @@ public class EpicFirebaseInstanceIDService extends FirebaseInstanceIdService {
 		}
 	}
 
+	private static String getFirebaseInstanceToken() {
+		FirebaseInstanceId id = getFirebaseInstanceId();
+		return (id == null) ? "" : id.getToken();
+	}
+	
 	private static FirebaseInstanceId getFirebaseInstanceId() {
-		FirebaseApp app = FirebaseApp.getInstance(MESSAGING_CONFIG);
-		return FirebaseInstanceId.getInstance(app);
+		try {
+			FirebaseApp app = FirebaseApp.getInstance(MESSAGING_CONFIG);
+			return FirebaseInstanceId.getInstance(app);
+		} catch (Exception e) {
+			Log.error("FirebaseApp with name " + MESSAGING_CONFIG + " doesn't exist");
+			return null;
+		}
 	}
 	
 	private static void saveFirebaseToken(@NonNull Context context, @NonNull String firebaseToken) {
@@ -80,7 +90,7 @@ public class EpicFirebaseInstanceIDService extends FirebaseInstanceIdService {
 		Log.debug("Firebase token retrieved from cache: " + token);
 		if(TextUtils.isEmpty(token)) {
 			// handle edge case where we missed onTokenRefresh - ex. App Upgrade
-			token = getFirebaseInstanceId().getToken();
+			token = getFirebaseInstanceToken();
 			if(!TextUtils.isEmpty(token)) {
 				Log.debug("Firebase token retrieved from Firebase: " + token);
 				saveFirebaseToken(context, token);
diff --git a/Engine/Source/Editor/Persona/Private/PersonaModule.cpp b/Engine/Source/Editor/Persona/Private/PersonaModule.cpp
--- a/Engine/Source/Editor/Persona/Private/PersonaModule.cpp
+++ b/Engine/Source/Editor/Persona/Private/PersonaModule.cpp
@@ -104,6 +104,9 @@ void FPersonaModule::StartupModule()
 	FModuleManager::Get().LoadModuleChecked("AdvancedPreviewScene");
 
 	// Load all blueprint animnotifies from asset registry so they are available from drop downs in anim segment detail views
+	FString Commandline = FCommandLine::Get();
+	bool bIsCookCommandlet = Commandline.Contains(TEXT("cookcommandlet")) || Commandline.Contains(TEXT("run=cook"));
+	if(!bIsCookCommandlet)
 	{
 		FAssetRegistryModule& AssetRegistryModule = FModuleManager::LoadModuleChecked<FAssetRegistryModule>(TEXT("AssetRegistry"));
 
diff --git a/Engine/Source/Programs/AutomationTool/IOS/IOSPlatform.Automation.cs b/Engine/Source/Programs/AutomationTool/IOS/IOSPlatform.Automation.cs
--- a/Engine/Source/Programs/AutomationTool/IOS/IOSPlatform.Automation.cs
+++ b/Engine/Source/Programs/AutomationTool/IOS/IOSPlatform.Automation.cs
@@ -676,6 +676,7 @@ public class IOSPlatform : Platform
             // Set encoding to support unicode filenames
             Zip.AlternateEncodingUsage = ZipOption.Always;
             Zip.AlternateEncoding = Encoding.UTF8;
+			Zip.UseZip64WhenSaving = Zip64Option.AsNecessary;
 
 			// set the compression level
 			bool bUseMaxIPACompression = ShouldUseMaxIPACompression(Params);
diff --git a/Engine/Source/Programs/UnrealBuildTool/Platform/Android/UEDeployAndroid.cs b/Engine/Source/Programs/UnrealBuildTool/Platform/Android/UEDeployAndroid.cs
--- a/Engine/Source/Programs/UnrealBuildTool/Platform/Android/UEDeployAndroid.cs
+++ b/Engine/Source/Programs/UnrealBuildTool/Platform/Android/UEDeployAndroid.cs
@@ -1410,6 +1410,11 @@ namespace UnrealBuildTool
 					Log.TraceInformation("{0}", e.Data.Replace("Warning: ", ">> "));
 					return;
 				}
+				if (e.Data.Contains("Unable to strip library"))
+				{
+					Log.TraceInformation("{0}", e.Data.Replace("due to error", ""));
+					return;
+				}
 				Log.TraceInformation("{0}", e.Data);
 			}
 		}
@@ -1785,16 +1790,28 @@ namespace UnrealBuildTool
 			}
 		}
 
-		private void DetermineScreenOrientationRequirements(out bool bNeedPortrait, out bool bNeedLandscape)
+		private string GetOrientation(string NDKArch)
 		{
 			ConfigHierarchy Ini = GetConfigCacheIni(ConfigHierarchyType.Engine);
 			string Orientation;
 			Ini.GetString("/Script/AndroidRuntimeSettings.AndroidRuntimeSettings", "Orientation", out Orientation);
 
+			// check for UPL override
+			string OrientationOverride = UPL.ProcessPluginNode(NDKArch, "orientationOverride", "");
+			if (!String.IsNullOrEmpty(OrientationOverride))
+			{
+				Orientation = OrientationOverride;
+			}
+
+			return ConvertOrientationIniValue(Orientation);
+		}
+
+		private void DetermineScreenOrientationRequirements(string Arch, out bool bNeedPortrait, out bool bNeedLandscape)
+		{
 			bNeedLandscape = false;
 			bNeedPortrait = false;
 
-			switch (Orientation.ToLower())
+			switch (GetOrientation(Arch).ToLower())
 			{
 				case "portrait":
 					bNeedPortrait = true;
@@ -2146,8 +2163,7 @@ namespace UnrealBuildTool
 			{
 				MaxAspectRatioValue = 2.1f;
 			}
-			string Orientation;
-			Ini.GetString("/Script/AndroidRuntimeSettings.AndroidRuntimeSettings", "Orientation", out Orientation);
+			string Orientation = ConvertOrientationIniValue(GetOrientation(Arch));
 			bool EnableFullScreen;
 			Ini.GetBool("/Script/AndroidRuntimeSettings.AndroidRuntimeSettings", "bFullScreen", out EnableFullScreen);
 			bool bUseDisplayCutout;
@@ -2371,7 +2387,7 @@ namespace UnrealBuildTool
 				Text.AppendLine("\t\t          android:label=\"@string/app_name\"");
 				Text.AppendLine("\t\t          android:theme=\"@style/UE4SplashTheme\"");
 				Text.AppendLine("\t\t          android:launchMode=\"singleTask\"");
-				Text.AppendLine(string.Format("\t\t          android:screenOrientation=\"{0}\"", ConvertOrientationIniValue(Orientation)));
+				Text.AppendLine(string.Format("\t\t          android:screenOrientation=\"{0}\"", Orientation));
 				Text.AppendLine(string.Format("\t\t          android:debuggable=\"{0}\">", bIsForDistribution ? "false" : "true"));
 				Text.AppendLine("\t\t\t<intent-filter>");
 				Text.AppendLine("\t\t\t\t<action android:name=\"android.intent.action.MAIN\" />");
@@ -2393,7 +2409,7 @@ namespace UnrealBuildTool
 											: "\t\t          android:configChanges=\"mcc|mnc|uiMode|screenSize|orientation|keyboardHidden|keyboard\"");
 			}
 			Text.AppendLine("\t\t          android:launchMode=\"singleTask\"");
-			Text.AppendLine(string.Format("\t\t          android:screenOrientation=\"{0}\"", ConvertOrientationIniValue(Orientation)));
+			Text.AppendLine(string.Format("\t\t          android:screenOrientation=\"{0}\"", Orientation));
 			if (ExtraActivityNodeTags != null)
 			{
 				foreach (string Line in ExtraActivityNodeTags)
@@ -2432,7 +2448,7 @@ namespace UnrealBuildTool
 			if (bShowLaunchImage)
 			{
 				Text.AppendLine("\t\t<activity android:name=\".DownloaderActivity\"");
-				Text.AppendLine(string.Format("\t\t          android:screenOrientation=\"{0}\"", ConvertOrientationIniValue(Orientation)));
+				Text.AppendLine(string.Format("\t\t          android:screenOrientation=\"{0}\"", Orientation));
 				Text.AppendLine(bAddDensity ? "\t\t          android:configChanges=\"mcc|mnc|uiMode|density|screenSize|orientation|keyboardHidden|keyboard\""
 											: "\t\t          android:configChanges=\"mcc|mnc|uiMode|screenSize|orientation|keyboardHidden|keyboard\"");
 				Text.AppendLine("\t\t          android:theme=\"@style/UE4SplashTheme\" />");
@@ -3243,20 +3259,43 @@ namespace UnrealBuildTool
 			string NDKBuildPath = Environment.ExpandEnvironmentVariables("%NDKROOT%/ndk-build" + (Utils.IsRunningOnMono ? "" : ".cmd"));
 
 			// set up some directory info
-			string IntermediateAndroidPath = Path.Combine(ProjectDirectory, "Intermediate/Android/");
+			string IntermediateAndroidPath = Path.Combine(ProjectDirectory, "Intermediate", "Android");
 			string UE4BuildPath = Path.Combine(IntermediateAndroidPath, "APK");
 			string UE4JavaFilePath = Path.Combine(ProjectDirectory, "Build", "Android", GetUE4JavaSrcPath());
 			string UE4BuildFilesPath = GetUE4BuildFilePath(EngineDirectory);
-			string GameBuildFilesPath = Path.Combine(ProjectDirectory, "Build/Android");
+			string GameBuildFilesPath = Path.Combine(ProjectDirectory, "Build", "Android");
 			string UE4BuildAssetsPath = Path.Combine(UE4BuildPath, "assets");
 			string UE4BuildResourcesPath = Path.Combine(UE4BuildPath, "res");
 
 			// force create from scratch if on build machine
-			if (bIsBuildMachine)
+			bool bCreateFromScratch = bIsBuildMachine;
+
+			// see if last time matches the skipGradle setting
+			string BuildTypeFilename = Path.Combine(IntermediateAndroidPath, "BuildType.txt");
+			string BuildTypeID = bSkipGradleBuild ? "Embedded" : "Standalone";
+			if (File.Exists(BuildTypeFilename))
 			{
+				string BuildTypeContents = File.ReadAllText(BuildTypeFilename);
+				if (BuildTypeID != BuildTypeContents)
+				{
+					bCreateFromScratch = true;
+				}
+			}
+
+			if (bCreateFromScratch)
+			{
+				Log.TraceInformation("Cleaning {0}", IntermediateAndroidPath);
 				DeleteDirectory(IntermediateAndroidPath);
 				Directory.CreateDirectory(IntermediateAndroidPath);
 			}
+			
+			if (!System.IO.Directory.Exists(IntermediateAndroidPath))
+			{
+				System.IO.Directory.CreateDirectory(IntermediateAndroidPath);
+			}
+
+			// write build type
+			File.WriteAllText(BuildTypeFilename, BuildTypeID);
 
 			// cache if we want data in the Apk
 			bool bPackageDataInsideApk = bDisallowPackagingDataInApk ? false : GetPackageDataInsideApk();
@@ -3585,7 +3624,7 @@ namespace UnrealBuildTool
 			//determine which orientation requirements this app has
 			bool bNeedLandscape = false;
 			bool bNeedPortrait = false;
-			DetermineScreenOrientationRequirements(out bNeedPortrait, out bNeedLandscape);
+			DetermineScreenOrientationRequirements(NDKArches[0], out bNeedPortrait, out bNeedLandscape);
 
 			//Now keep the splash screen images matching orientation requested
 			PickSplashScreenOrientation(UE4BuildPath, bNeedPortrait, bNeedLandscape);
@@ -4214,7 +4253,9 @@ bSaveSymbols = true;
 				{ "//$${gameActivityOnResumeAdditions}$$", UPL.ProcessPluginNode(NDKArch, "gameActivityOnResumeAdditions", "")},
 				{ "//$${gameActivityOnNewIntentAdditions}$$", UPL.ProcessPluginNode(NDKArch, "gameActivityOnNewIntentAdditions", "")},
   				{ "//$${gameActivityOnActivityResultAdditions}$$", UPL.ProcessPluginNode(NDKArch, "gameActivityOnActivityResultAdditions", "")},
+  				{ "//$${gameActivityPreConfigRulesParseAdditions}$$", UPL.ProcessPluginNode(NDKArch, "gameActivityPreConfigRulesParseAdditions", "")},
   				{ "//$${gameActivityPostConfigRulesAdditions}$$", UPL.ProcessPluginNode(NDKArch, "gameActivityPostConfigRulesAdditions", "")},
+  				{ "//$${gameActivityFinalizeConfigRulesAdditions}$$", UPL.ProcessPluginNode(NDKArch, "gameActivityFinalizeConfigRulesAdditions", "")},
 				{ "//$${gameActivityBeforeConfigRulesAppliedAdditions}$$", UPL.ProcessPluginNode(NDKArch, "gameActivityBeforeConfigRulesAppliedAdditions", "")},
 				{ "//$${gameActivityLoggerCallbackAdditions}$$", UPL.ProcessPluginNode(NDKArch, "gameActivityLoggerCallbackAdditions", "")},
 				{ "//$${soLoadLibrary}$$", UPL.ProcessPluginNode(NDKArch, "soLoadLibrary", LoadLibraryDefaults)},
@@ -4263,8 +4304,8 @@ bSaveSymbols = true;
 			string DestFilename = Path.Combine(UE4BuildPath, "src", "com", "epicgames", "ue4", "GameApplication.java");
 
 			Dictionary<string, string> Replacements = new Dictionary<string, string>{
-				{ "//$${gameActivityImportAdditions}$$", UPL.ProcessPluginNode(NDKArch, "gameApplicationImportAdditions", "")},
-				{ "//$${gameActivityOnCreateAdditions}$$", UPL.ProcessPluginNode(NDKArch, "gameApplicationOnCreateAdditions", "")},
+				{ "//$${gameApplicationImportAdditions}$$", UPL.ProcessPluginNode(NDKArch, "gameApplicationImportAdditions", "")},
+				{ "//$${gameApplicationOnCreateAdditions}$$", UPL.ProcessPluginNode(NDKArch, "gameApplicationOnCreateAdditions", "")},
 			};
 
 			string[] TemplateSrc = File.ReadAllLines(SourceFilename);
@@ -4430,9 +4471,9 @@ bSaveSymbols = true;
 				string GradleProjectPath = Path.Combine(UE4BuildGradlePath, RelativePath);
 				string GradleProjectMainPath = Path.Combine(GradleProjectPath, "src", "main");
 
+				string ManifestFilename = Path.Combine(LibDir, "AndroidManifest.xml");
 				string GradleManifest = Path.Combine(GradleProjectMainPath, "AndroidManifest.xml");
 				MakeDirectoryIfRequired(GradleManifest);
-				File.Copy(Path.Combine(LibDir, "AndroidManifest.xml"), GradleManifest, true);
 
 				// Copy parts were they need to be
 				CleanCopyDirectory(Path.Combine(LibDir, "assets"), Path.Combine(GradleProjectPath, "assets"));
@@ -4460,7 +4501,8 @@ bSaveSymbols = true;
 				// Try to get the SDK target from the AndroidManifest.xml
 				string VersionCode = "";
 				string VersionName = "";
-				string ManifestFilename = Path.Combine(LibDir, "AndroidManifest.xml");
+				string MinSdkVersion = "";
+				string TargetSdkVersion = "";
 				XDocument ManifestXML;
 				if (File.Exists(ManifestFilename))
 				{
@@ -4479,6 +4521,41 @@ bSaveSymbols = true;
 						{
 							VersionName = VersionNameAttr.Value;
 						}
+
+						XElement UseSDKNode = null;
+						foreach (XElement WorkNode in ManifestXML.Elements().First().Descendants("uses-sdk"))
+						{
+							UseSDKNode = WorkNode;
+
+							XAttribute MinSdkVersionAttr = WorkNode.Attribute(XName.Get("minSdkVersion", "http://schemas.android.com/apk/res/android"));
+							if (MinSdkVersionAttr != null)
+							{
+								MinSdkVersion = MinSdkVersionAttr.Value;
+							}
+
+							XAttribute TargetSdkVersionAttr = WorkNode.Attribute(XName.Get("targetSdkVersion", "http://schemas.android.com/apk/res/android"));
+							if (TargetSdkVersionAttr != null)
+							{
+								TargetSdkVersion = TargetSdkVersionAttr.Value;
+							}
+						}
+
+						if (UseSDKNode != null)
+						{
+							UseSDKNode.Remove();
+						}
+
+						// rewrite the manifest if different
+						String NewManifestText = ManifestXML.ToString();
+						String OldManifestText = "";
+						if (File.Exists(GradleManifest))
+						{
+							OldManifestText = File.ReadAllText(GradleManifest);
+						}
+						if (NewManifestText != OldManifestText)
+						{
+							File.WriteAllText(GradleManifest, NewManifestText);
+						}
 					}
 					catch (Exception e)
 					{
@@ -4494,6 +4571,14 @@ bSaveSymbols = true;
 				{
 					BuildGradleContent.AppendLine(string.Format("\t\tversionName \"{0}\"", VersionName));
 				}
+				if (MinSdkVersion != "")
+				{
+					BuildGradleContent.AppendLine(string.Format("\t\tminSdkVersion = {0}", MinSdkVersion));
+				}
+				if (TargetSdkVersion != "")
+				{
+					BuildGradleContent.AppendLine(string.Format("\t\ttargetSdkVersion = {0}", TargetSdkVersion));
+				}
 				BuildGradleContent.AppendLine("\t}");
 				BuildGradleContent.AppendLine("}");
 
@@ -4507,8 +4592,18 @@ bSaveSymbols = true;
 					}
 				}
 
+				// rewrite the build.gradle if different
 				string BuildGradleFilename = Path.Combine(GradleProjectPath, "build.gradle");
-				File.WriteAllText(BuildGradleFilename, BuildGradleContent.ToString());
+				String NewBuildGradleText = BuildGradleContent.ToString();
+				String OldBuildGradleText = "";
+				if (File.Exists(BuildGradleFilename))
+				{
+					OldBuildGradleText = File.ReadAllText(BuildGradleFilename);
+				}
+				if (NewBuildGradleText != OldBuildGradleText)
+				{
+					File.WriteAllText(BuildGradleFilename, NewBuildGradleText);
+				}
 			}
 			ProjectDependencyContent.AppendLine("}");
 
diff --git a/Engine/Source/Programs/UnrealBuildTool/Platform/IOS/IOSToolChain.cs b/Engine/Source/Programs/UnrealBuildTool/Platform/IOS/IOSToolChain.cs
--- a/Engine/Source/Programs/UnrealBuildTool/Platform/IOS/IOSToolChain.cs
+++ b/Engine/Source/Programs/UnrealBuildTool/Platform/IOS/IOSToolChain.cs
@@ -1577,7 +1577,13 @@ namespace UnrealBuildTool
 					{
 						throw new BuildException("Unable to find plist for output framework ({0})", PlistSrcLocation);
 					}
+					
 					FileReference PlistDstLocation = FileReference.Combine(Target.OutputPath.Directory, "Info.plist");
+					if (FileReference.Exists(PlistDstLocation))
+					{
+						FileReference.SetAttributes(PlistDstLocation, FileAttributes.Normal);
+					}
+
 					FileReference.Copy(PlistSrcLocation, PlistDstLocation, true);
 
 					// and do nothing else
diff --git a/Engine/Source/Runtime/Core/Private/Android/AndroidPlatformStackWalk.cpp b/Engine/Source/Runtime/Core/Private/Android/AndroidPlatformStackWalk.cpp
--- a/Engine/Source/Runtime/Core/Private/Android/AndroidPlatformStackWalk.cpp
+++ b/Engine/Source/Runtime/Core/Private/Android/AndroidPlatformStackWalk.cpp
@@ -131,11 +131,12 @@ bool FAndroidPlatformStackWalk::SymbolInfoToHumanReadableString(const FProgramCo
 	const int32 MAX_TEMP_SPRINTF = 256;
 
 	//
-	// Callstack lines should be written in this standard format
+	// Callstack lines should be written in this standard format. These are parsed by tools so it is important that
+	// extra elements are not inserted!
 	//
 	//	0xaddress module!func [file]
 	// 
-	// E.g. 0x045C8D01 (0x00009034) OrionClient.self!UEngine::PerformError() [D:\Epic\Orion\Engine\Source\Runtime\Engine\Private\UnrealEngine.cpp:6481]
+	// E.g. 0x045C8D01 OrionClient.self(0x00009034)!UEngine::PerformError() [D:\Epic\Orion\Engine\Source\Runtime\Engine\Private\UnrealEngine.cpp:6481]
 	//
 	// Module may be omitted, everything else should be present, or substituted with a string that conforms to the expected type
 	//
@@ -156,17 +157,18 @@ bool FAndroidPlatformStackWalk::SymbolInfoToHumanReadableString(const FProgramCo
 		ANSICHAR PCAddress[MAX_TEMP_SPRINTF] = { 0 };
 		FCStringAnsi::Snprintf(PCAddress, MAX_TEMP_SPRINTF, "0x%016llX ", SymbolInfo.ProgramCounter);
 		FCStringAnsi::Strncat(StackLine, PCAddress, MAX_SPRINTF);
-		FCStringAnsi::Snprintf(PCAddress, MAX_TEMP_SPRINTF, "(0x%016llX) ", SymbolInfo.OffsetInModule);
-		FCStringAnsi::Strncat(StackLine, PCAddress, MAX_SPRINTF);
 
 		// Module if it's present
 		const bool bHasValidModuleName = FCStringAnsi::Strlen(StrippedModuleName) > 0;
 		if (bHasValidModuleName)
 		{
 			FCStringAnsi::Strncat(StackLine, StrippedModuleName, MAX_SPRINTF);
-			FCStringAnsi::Strncat(StackLine, "!", MAX_SPRINTF);
 		}
 
+		FCStringAnsi::Snprintf(PCAddress, MAX_TEMP_SPRINTF, "(0x%016llX)", SymbolInfo.OffsetInModule);
+		FCStringAnsi::Strncat(StackLine, PCAddress, MAX_SPRINTF);
+		FCStringAnsi::Strncat(StackLine, "!", MAX_SPRINTF);
+
 		// Function if it's available, unknown if it's not
 		const bool bHasValidFunctionName = FCStringAnsi::Strlen(SymbolInfo.FunctionName) > 0;
 		if (bHasValidFunctionName)
diff --git a/Engine/Source/Runtime/Core/Private/GenericPlatform/GenericPlatformMisc.cpp b/Engine/Source/Runtime/Core/Private/GenericPlatform/GenericPlatformMisc.cpp
--- a/Engine/Source/Runtime/Core/Private/GenericPlatform/GenericPlatformMisc.cpp
+++ b/Engine/Source/Runtime/Core/Private/GenericPlatform/GenericPlatformMisc.cpp
@@ -834,13 +834,10 @@ IPlatformInstallBundleManager* FGenericPlatformMisc::GetPlatformInstallBundleMan
 		GConfig->GetString(TEXT("InstallBundleManager"), TEXT("ModuleName"), ModuleName, GEngineIni);
 
 		FModuleStatus Status;
-		if (FModuleManager::Get().QueryModule(*ModuleName, Status))
+		Module = FModuleManager::LoadModulePtr<IPlatformInstallBundleManagerModule>(*ModuleName);
+		if (Module)
 		{
-			Module = FModuleManager::LoadModulePtr<IPlatformInstallBundleManagerModule>(*ModuleName);
-			if (Module)
-			{
-				Manager = Module->GetInstallBundleManager();
-			}
+			Manager = Module->GetInstallBundleManager();
 		}
 		bCheckedIni = true;
 	}
diff --git a/Engine/Source/Runtime/Core/Private/Mac/MacPlatformMisc.cpp b/Engine/Source/Runtime/Core/Private/Mac/MacPlatformMisc.cpp
--- a/Engine/Source/Runtime/Core/Private/Mac/MacPlatformMisc.cpp
+++ b/Engine/Source/Runtime/Core/Private/Mac/MacPlatformMisc.cpp
@@ -631,9 +631,11 @@ int32 FMacPlatformMisc::NumberOfCores()
 		else
 		{
 			SIZE_T Size = sizeof(int32);
+			int Result = sysctlbyname("hw.physicalcpu", &NumberOfCores, &Size, NULL, 0);
 		
-			if (sysctlbyname("hw.physicalcpu", &NumberOfCores, &Size, NULL, 0) != 0)
+			if (Result != 0)
 			{
+				UE_LOG(LogMac, Error,  TEXT("sysctlbyname(hw.physicalcpu...) failed with error %d. Defaulting to one core"), Result);
 				NumberOfCores = 1;
 			}
 		}
@@ -643,7 +645,19 @@ int32 FMacPlatformMisc::NumberOfCores()
 
 int32 FMacPlatformMisc::NumberOfCoresIncludingHyperthreads()
 {
-	return FApplePlatformMisc::NumberOfCores();
+	static int32 NumberOfCores = -1;
+	if (NumberOfCores == -1)
+	{
+		SIZE_T Size = sizeof(int32);
+		int Result = sysctlbyname("hw.logicalcpu", &NumberOfCores, &Size, NULL, 0);
+		
+		if (Result != 0)
+		{
+			UE_LOG(LogMac, Error,  TEXT("sysctlbyname(hw.logicalcpu...) failed with error %d. Defaulting to one core"), Result);
+			NumberOfCores = 1;
+		}
+	}
+	return NumberOfCores;
 }
 
 void FMacPlatformMisc::NormalizePath(FString& InPath)
diff --git a/Engine/Source/Runtime/Core/Public/GenericPlatform/GenericPlatformInstallBundleManager.h b/Engine/Source/Runtime/Core/Public/GenericPlatform/GenericPlatformInstallBundleManager.h
--- a/Engine/Source/Runtime/Core/Public/GenericPlatform/GenericPlatformInstallBundleManager.h
+++ b/Engine/Source/Runtime/Core/Public/GenericPlatform/GenericPlatformInstallBundleManager.h
@@ -3,13 +3,16 @@
 #pragma once
 
 #include "Modules/ModuleInterface.h"
+#include "Modules/ModuleManager.h"
 #include "Misc/EnumClassFlags.h"
+#include "Misc/ConfigCacheIni.h"
 #include "Logging/LogMacros.h"
 #include "Internationalization/Text.h"
 
 enum class EInstallBundleResult : int
 {
 	OK,
+	FailedPrereqRequiresLatestClient,
 	InstallError,
 	InstallerOutOfDiskSpaceError,
 	OnCellularNetworkError,
@@ -25,6 +28,7 @@ inline const TCHAR* GetInstallBundleResultString(EInstallBundleResult Result)
 	static const TCHAR* Strings[] =
 	{
 		TEXT("OK"),
+		TEXT("FailedPrereqRequiresLatestClient"),
 		TEXT("InstallError"),
 		TEXT("InstallerOutOfDiskSpaceError"),
 		TEXT("OnCellularNetworkError"),
@@ -119,6 +123,19 @@ struct FInstallBundleResultInfo
 	FString OptionalErrorCode;
 };
 
+enum class EInstallBundleContentState : int
+{
+	InitializationError,
+	NeedsUpdate,
+	UpToDate,
+};
+
+struct FInstallBundleContentState
+{
+	EInstallBundleContentState State = EInstallBundleContentState::InitializationError;
+	uint64 DownloadSize = 0;
+};
+
 enum class EInstallBundleRequestInfoFlags : int32
 {
 	None							= 0,
@@ -150,6 +167,8 @@ DECLARE_DELEGATE_RetVal_OneParam(EInstallBundleManagerInitErrorHandlerResult, FI
 
 DECLARE_MULTICAST_DELEGATE_OneParam(FInstallBundleCompleteMultiDelegate, FInstallBundleResultInfo);
 
+DECLARE_DELEGATE_OneParam(FInstallBundleGetContentStateDelegate, FInstallBundleContentState);
+
 class CORE_API IPlatformInstallBundleManager
 {
 public:
@@ -169,6 +188,9 @@ public:
 	virtual FInstallBundleRequestInfo RequestUpdateContent(FName BundleName, EInstallBundleRequestFlags Flags) = 0;
 	virtual FInstallBundleRequestInfo RequestUpdateContent(TArrayView<FName> BundleNames, EInstallBundleRequestFlags Flags) = 0;
 
+	virtual void GetContentState(FName BundleName, bool bAddDependencies, FInstallBundleGetContentStateDelegate Callback) = 0;
+	virtual void GetContentState(TArrayView<FName> BundleNames, bool bAddDependencies, FInstallBundleGetContentStateDelegate Callback) = 0;
+
 	virtual FInstallBundleRequestInfo RequestRemoveBundle(FName BundleName) = 0;
 
 	virtual void RequestRemoveBundleOnNextInit(FName BundleName) = 0;
@@ -181,5 +203,33 @@ public:
 class IPlatformInstallBundleManagerModule : public IModuleInterface
 {
 public:
-	virtual IPlatformInstallBundleManager* GetInstallBundleManager() = 0;
+	virtual void PreUnloadCallback() override
+	{
+		InstallBundleManager.Reset();
+	}
+
+	IPlatformInstallBundleManager* GetInstallBundleManager()
+	{
+		return InstallBundleManager.Get();
+	}
+
+protected:
+	TUniquePtr<IPlatformInstallBundleManager> InstallBundleManager;
+};
+
+template<class PlatformInstallBundleManagerImpl>
+class TPlatformInstallBundleManagerModule : public IPlatformInstallBundleManagerModule
+{
+public:
+	virtual void StartupModule() override
+	{
+		// Only instantiate the bundle manager if this is the version the game has been configured to use
+		FString ModuleName;
+		GConfig->GetString(TEXT("InstallBundleManager"), TEXT("ModuleName"), ModuleName, GEngineIni);
+
+		if (FModuleManager::Get().GetModule(*ModuleName) == this)
+		{
+			InstallBundleManager = MakeUnique<PlatformInstallBundleManagerImpl>();
+		}
+	}
 };
diff --git a/Engine/Source/Runtime/Engine/Private/AssetManager.cpp b/Engine/Source/Runtime/Engine/Private/AssetManager.cpp
--- a/Engine/Source/Runtime/Engine/Private/AssetManager.cpp
+++ b/Engine/Source/Runtime/Engine/Private/AssetManager.cpp
@@ -508,11 +508,29 @@ int32 UAssetManager::ScanPathsForPrimaryAssets(FPrimaryAssetType PrimaryAssetTyp
 			}
 		}
 
+
 		FPrimaryAssetId PrimaryAssetId = ExtractPrimaryAssetIdFromData(Data, PrimaryAssetType);
 
 		// Remove invalid or wrong type assets
 		if (!PrimaryAssetId.IsValid() || PrimaryAssetId.PrimaryAssetType != PrimaryAssetType)
 		{
+			if (!PrimaryAssetId.IsValid())
+			{
+				UE_LOG(LogAssetManager, Warning, TEXT("Ignoring primary asset %s - PrimaryAssetType %s - invalid primary asset ID"), *Data.AssetName.ToString(), *PrimaryAssetType.ToString());
+			}
+			else
+			{
+				// Warn that 'Foo' conflicts with 'Bar', but only once per conflict
+				static TSet<FString> IssuedWarnings;
+
+				FString ConflictMsg = FString::Printf(TEXT("Ignoring PrimaryAssetType %s - Conflicts with %s"), *PrimaryAssetType.ToString(), *PrimaryAssetId.PrimaryAssetType.ToString());
+
+				if (!IssuedWarnings.Contains(ConflictMsg))
+				{
+					UE_LOG(LogAssetManager, Warning, TEXT("%s"), *ConflictMsg);
+					IssuedWarnings.Add(ConflictMsg);
+				}
+			}
 			continue;
 		}
 
diff --git a/Engine/Source/Runtime/Engine/Private/GameEngine.cpp b/Engine/Source/Runtime/Engine/Private/GameEngine.cpp
--- a/Engine/Source/Runtime/Engine/Private/GameEngine.cpp
+++ b/Engine/Source/Runtime/Engine/Private/GameEngine.cpp
@@ -762,6 +762,16 @@ public:
 			// execute any console commands
 			if (Message.Command == TEXT("console"))
 			{
+				// executed too early, must fail
+				if (GEngine == nullptr)
+				{
+					// call the completion delegate with all text output
+					TArray<FString> Keys;
+					Message.Parameters.GetKeys(Keys);
+					Message.OnCompleteDelegate({ }, FString::Printf(TEXT("GEngine does not exist yet, unable to execute console command (%s)"), Keys.Num() ? *Keys[0] : TEXT("None")));
+					return;
+				}
+				
 				// gather all of the output
 				FStringOutputDevice Output;
 				ULocalPlayer* LocalPlayer = GEngine->GetDebugLocalPlayer();
diff --git a/Engine/Source/Runtime/GameplayMediaEncoder/Private/GameplayMediaEncoder.cpp b/Engine/Source/Runtime/GameplayMediaEncoder/Private/GameplayMediaEncoder.cpp
--- a/Engine/Source/Runtime/GameplayMediaEncoder/Private/GameplayMediaEncoder.cpp
+++ b/Engine/Source/Runtime/GameplayMediaEncoder/Private/GameplayMediaEncoder.cpp
@@ -199,7 +199,6 @@ bool FGameplayMediaEncoder::Initialize()
 	}
 
 	// Specifying 0 will completely disable frame skipping (therefore encoding as many frames as possible)
-	VideoConfig.Framerate = HardcodedAudioSamplerate;
 	FParse::Value(FCommandLine::Get(), TEXT("GameplayMediaEncoder.FPS="), VideoConfig.Framerate);
 	if (VideoConfig.Framerate == 0)
 	{
diff --git a/Engine/Source/Runtime/Launch/Private/Android/LaunchAndroid.cpp b/Engine/Source/Runtime/Launch/Private/Android/LaunchAndroid.cpp
--- a/Engine/Source/Runtime/Launch/Private/Android/LaunchAndroid.cpp
+++ b/Engine/Source/Runtime/Launch/Private/Android/LaunchAndroid.cpp
@@ -523,7 +523,7 @@ struct FChoreographer
 	void SetupChoreographer()
 	{
 		FScopeLock Lock(&ChoreographerSetupLock);
-		check(!AChoreographer_getInstance_);
+		//check(!AChoreographer_getInstance_);
 		if (!AChoreographer_getInstance_)
 		{
 			void* lib = dlopen("libandroid.so", RTLD_NOW | RTLD_LOCAL);
diff --git a/Engine/Source/Runtime/Launch/Private/LaunchEngineLoop.cpp b/Engine/Source/Runtime/Launch/Private/LaunchEngineLoop.cpp
--- a/Engine/Source/Runtime/Launch/Private/LaunchEngineLoop.cpp
+++ b/Engine/Source/Runtime/Launch/Private/LaunchEngineLoop.cpp
@@ -215,16 +215,8 @@ class FFeedbackContext;
 	#define RHI_COMMAND_LIST_DEBUG_TRACES 0
 #endif
 
-#ifndef REAPPLY_INI_SETTINGS_AFTER_EARLY_LOADING_SCREEN
-#define REAPPLY_INI_SETTINGS_AFTER_EARLY_LOADING_SCREEN 0
-#endif
-
-#ifndef MOUNT_DOWNLOADED_PAKS_AFTER_EARLY_LOADING_SCREEN
-#define MOUNT_DOWNLOADED_PAKS_AFTER_EARLY_LOADING_SCREEN 1
-#endif
-
-#ifndef OPEN_PSO_CACHE_AFTER_EARLY_LOADING_SCREEN
-#define OPEN_PSO_CACHE_AFTER_EARLY_LOADING_SCREEN 1
+#ifndef AUTOMATICALLY_HANDLE_INSTALLED_CONTENT_AFTER_EARLY_STARTUP_SCREEN
+#define AUTOMATICALLY_HANDLE_INSTALLED_CONTENT_AFTER_EARLY_STARTUP_SCREEN 0
 #endif
 
 #if WITH_ENGINE
@@ -1587,8 +1579,8 @@ int32 FEngineLoop::PreInit(const TCHAR* CmdLine)
 		}
 	}
 
-	const bool bShouldReapplyCVarsFromIniAfterLoadScreen = REAPPLY_INI_SETTINGS_AFTER_EARLY_LOADING_SCREEN;
-	if (bShouldReapplyCVarsFromIniAfterLoadScreen)
+	const bool bShouldHandleContentInstalledDuringEarlyLoadScreen = AUTOMATICALLY_HANDLE_INSTALLED_CONTENT_AFTER_EARLY_STARTUP_SCREEN;
+	if (bShouldHandleContentInstalledDuringEarlyLoadScreen)
 	{
 		UE_LOG(LogInit, Verbose, TEXT("Reapplying ini settings after early loading screen."));
 
@@ -2219,6 +2211,24 @@ int32 FEngineLoop::PreInit(const TCHAR* CmdLine)
 				}
 			}
 
+			//If we are NOT expecting to automatically handle content after our EarlyStartupScreen, we can go ahead and load our game shaders now.
+			//If present they will be loaded, and if not present this will gracefully fail and we can handle loading them in our EarlyStartupScreen.
+			if (!bShouldHandleContentInstalledDuringEarlyLoadScreen)
+			{
+				LLM_SCOPE(ELLMTag::Shaders);
+				SCOPED_BOOT_TIMING("FShaderCodeLibrary::OpenLibrary");
+
+				// Open the game library which contains the material shaders.
+				FShaderCodeLibrary::OpenLibrary(FApp::GetProjectName(), FPaths::ProjectContentDir());
+				for (const FString& RootDir : FPlatformMisc::GetAdditionalRootDirectories())
+				{
+					FShaderCodeLibrary::OpenLibrary(FApp::GetProjectName(), FPaths::Combine(RootDir, FApp::GetProjectName(), TEXT("Content")));
+				}
+
+				// Now our shader code main library is opened, kick off the precompile.
+				FShaderPipelineCache::OpenPipelineFileCache(GMaxRHIShaderPlatform);
+			}
+
 			if (GetMoviePlayer()->HasEarlyStartupMovie())
 			{
 				SCOPED_BOOT_TIMING("EarlyStartupMovie");
@@ -2323,56 +2333,51 @@ int32 FEngineLoop::PreInit(const TCHAR* CmdLine)
 			FEngineFontServices::Create();
 		}
 #endif
-		if (bShouldReapplyCVarsFromIniAfterLoadScreen)
-		{
-			SCOPED_BOOT_TIMING("ReapplyCVarsFromIniAfterLoadScreen");
 
-			// This can be deleted after ForInstallBundleManager is in Engine
-			const bool bMountPaks = MOUNT_DOWNLOADED_PAKS_AFTER_EARLY_LOADING_SCREEN;
-			if (bMountPaks && FCoreDelegates::OnMountAllPakFiles.IsBound() )
+		//Now that our EarlyStartupScreen is finished, lets take the necessary steps to mount paks, apply .ini cvars, and open the shader libraries if we installed content we expect to handle
+		if (bShouldHandleContentInstalledDuringEarlyLoadScreen)
+		{
+			// Mount Paks that were installed during EarlyStartupScreen
+			if (FCoreDelegates::OnMountAllPakFiles.IsBound() )
 			{
-#if 1
+				SCOPED_BOOT_TIMING("MountPaksAfterEarlyStartupScreen");
+
 				FString InstalledGameContentDir = FPaths::Combine(*FPaths::ProjectPersistentDownloadDir(), TEXT("InstalledContent"), FApp::GetProjectName(), TEXT("Content"), TEXT("Paks"));
 				FPlatformMisc::AddAdditionalRootDirectory(FPaths::Combine(*FPaths::ProjectPersistentDownloadDir(), TEXT("InstalledContent")));
 
 				TArray<FString> PakFolders;
 				PakFolders.Add(InstalledGameContentDir);
 				FCoreDelegates::OnMountAllPakFiles.Execute(PakFolders);
-#else
-				TArray<FString> PakFolders;
-				PakFolders.Add(FPaths::Combine(*FPaths::ProjectPersistentDownloadDir(), TEXT("InstalledContent"), FApp::GetProjectName(), TEXT("Content"), TEXT("Paks")));
-				FCoreDelegates::OnMountAllPakFiles.Execute(PakFolders);
-#endif
 			}
 
-			extern CORE_API void ReapplyRecordedCVarSettingsFromIni();
-			extern CORE_API void DeleteRecordedCVarSettingsFromIni();
-
-			ReapplyRecordedCVarSettingsFromIni();
-			DeleteRecordedCVarSettingsFromIni();
-
-		}
+			//Reapply CVars after our EarlyLoadScreen
+			{
+				SCOPED_BOOT_TIMING("ReapplyCVarsFromIniAfterEarlyStartupScreen");
 
-		if (FPlatformProperties::RequiresCookedData())
-		{
-			LLM_SCOPE(ELLMTag::Shaders);
-			SCOPED_BOOT_TIMING("FShaderCodeLibrary::OpenLibrary");
+				extern CORE_API void ReapplyRecordedCVarSettingsFromIni();
+				extern CORE_API void DeleteRecordedCVarSettingsFromIni();
 
-			// Open the game library which contains the material shaders.
-			FShaderCodeLibrary::OpenLibrary(FApp::GetProjectName(), FPaths::ProjectContentDir());
-			for (const FString& RootDir : FPlatformMisc::GetAdditionalRootDirectories())
-			{
-				FShaderCodeLibrary::OpenLibrary(FApp::GetProjectName(), FPaths::Combine(RootDir, FApp::GetProjectName(), TEXT("Content")));
+				ReapplyRecordedCVarSettingsFromIni();
+				DeleteRecordedCVarSettingsFromIni();
 			}
 
-			const bool bOpenPSOCache = OPEN_PSO_CACHE_AFTER_EARLY_LOADING_SCREEN;
-			if (bOpenPSOCache)
+			//Handle opening shader library after our EarlyLoadScreen
 			{
+				LLM_SCOPE(ELLMTag::Shaders);
+				SCOPED_BOOT_TIMING("FShaderCodeLibrary::OpenLibrary");
+
+				// Open the game library which contains the material shaders.
+				FShaderCodeLibrary::OpenLibrary(FApp::GetProjectName(), FPaths::ProjectContentDir());
+				for (const FString& RootDir : FPlatformMisc::GetAdditionalRootDirectories())
+				{
+					FShaderCodeLibrary::OpenLibrary(FApp::GetProjectName(), FPaths::Combine(RootDir, FApp::GetProjectName(), TEXT("Content")));
+				}
+
 				// Now our shader code main library is opened, kick off the precompile.
 				FShaderPipelineCache::OpenPipelineFileCache(GMaxRHIShaderPlatform);
 			}
 		}
-		
+
 		InitGameTextLocalization();
 
 		DECLARE_SCOPE_CYCLE_COUNTER(TEXT("Initial UObject load"), STAT_InitialUObjectLoad, STATGROUP_LoadTime);
diff --git a/Engine/Source/Runtime/NullInstallBundleManager/Source/NullInstallBundleManager.cpp b/Engine/Source/Runtime/NullInstallBundleManager/Source/NullInstallBundleManager.cpp
--- a/Engine/Source/Runtime/NullInstallBundleManager/Source/NullInstallBundleManager.cpp
+++ b/Engine/Source/Runtime/NullInstallBundleManager/Source/NullInstallBundleManager.cpp
@@ -45,6 +45,20 @@ class FNullInstallBundleManager : public IPlatformInstallBundleManager
 		return RetInfo;
 	}
 
+	virtual void GetContentState(FName BundleName, bool bAddDependencies, FInstallBundleGetContentStateDelegate Callback) override
+	{
+		FInstallBundleContentState State;
+		State.State = EInstallBundleContentState::UpToDate;
+		Callback.ExecuteIfBound(State);
+	}
+
+	virtual void GetContentState(TArrayView<FName> BundleNames, bool bAddDependencies, FInstallBundleGetContentStateDelegate Callback) override
+	{
+		FInstallBundleContentState State;
+		State.State = EInstallBundleContentState::UpToDate;
+		Callback.ExecuteIfBound(State);
+	}
+
 	virtual void RequestRemoveBundleOnNextInit(FName BundleName) override
 	{
 
@@ -64,26 +78,8 @@ private:
 	
 };
 
-class FNullInstallBundleManagerModule : public IPlatformInstallBundleManagerModule
-{
-public:
-	virtual void StartupModule() override
-	{
-		InstallBundleManager = MakeUnique<FNullInstallBundleManager>();
-	}
-
-	virtual void PreUnloadCallback() override
-	{
-		InstallBundleManager.Reset();
-	}
-
-	virtual IPlatformInstallBundleManager* GetInstallBundleManager() override
-	{
-		return InstallBundleManager.Get();
-	}
-
-private:
-	TUniquePtr<FNullInstallBundleManager> InstallBundleManager;
+class FNullInstallBundleManagerModule : public TPlatformInstallBundleManagerModule<FNullInstallBundleManager>
+{	
 };
 
 IMPLEMENT_MODULE(FNullInstallBundleManagerModule, NullInstallBundleManager);
diff --git a/Engine/Source/Runtime/Slate/Private/Framework/Text/IOS/IOSPlatformTextField.cpp b/Engine/Source/Runtime/Slate/Private/Framework/Text/IOS/IOSPlatformTextField.cpp
--- a/Engine/Source/Runtime/Slate/Private/Framework/Text/IOS/IOSPlatformTextField.cpp
+++ b/Engine/Source/Runtime/Slate/Private/Framework/Text/IOS/IOSPlatformTextField.cpp
@@ -91,9 +91,15 @@ void FIOSPlatformTextField::ShowVirtualKeyboard(bool bShow, int32 UserIndex, TSh
 				TextField = [[[SlateTextField alloc] init] retain];
 			}
 
+			// capture some gamethread strings before we toss over to main thread
+			NSString* TextContents = [NSString stringWithFString : TextEntryWidget->GetText().ToString()];
+			NSString* PlaceholderContents = [NSString stringWithFString : TextEntryWidget->GetHintText().ToString()];
+			FKeyboardConfig KeyboardConfig;
+			GetKeyboardConfig(TextEntryWidget, KeyboardConfig);
+
 			// these functions must be run on the main thread
 			dispatch_async(dispatch_get_main_queue(),^ {
-				[TextField show: TextEntryWidget];
+				[TextField show:TextEntryWidget text:TextContents placeholder:PlaceholderContents keyboardConfig:KeyboardConfig];
 			});
 		}
         else
@@ -126,21 +132,10 @@ void FIOSPlatformTextField::ShowVirtualKeyboard(bool bShow, int32 UserIndex, TSh
         return;
     }
     
-#ifdef __IPHONE_8_0
-    if([UIAlertController class] && AlertController != nil)
+    if(AlertController != nil)
     {
         [AlertController dismissViewControllerAnimated: YES completion: nil];
     }
-    else
-#endif
-    {
-#if __IPHONE_OS_VERSION_MIN_REQUIRED < __IPHONE_9_0
-		if (AlertView != nil)
-        {
-            [AlertView dismissWithClickedButtonIndex: 0 animated: YES]; //0 is the cancel button
-        }
-#endif
-    }
     
     TextWidget = nullptr;
 }
@@ -150,176 +145,79 @@ void FIOSPlatformTextField::ShowVirtualKeyboard(bool bShow, int32 UserIndex, TSh
     return TextWidget.IsValid();
 }
 
--(void)show:(TSharedPtr<IVirtualKeyboardEntry>)InTextWidget
+-(void)show:(TSharedPtr<IVirtualKeyboardEntry>)InTextWidget text:(NSString*)TextContents placeholder:(NSString*)PlaceholderContents keyboardConfig:(FKeyboardConfig)KeyboardConfig
 {
 	TextWidget = InTextWidget;
 	TextEntry = FText::FromString(TEXT(""));
 
-#ifdef __IPHONE_8_0
-	if ([UIAlertController class])
-	{
-		AlertController = [UIAlertController alertControllerWithTitle : @"" message:@"" preferredStyle:UIAlertControllerStyleAlert];
-		UIAlertAction* okAction = [UIAlertAction 
-										actionWithTitle:NSLocalizedString(@"OK", nil)
-										style:UIAlertActionStyleDefault
-										handler:^(UIAlertAction* action)
+	AlertController = [UIAlertController alertControllerWithTitle : @"" message:@"" preferredStyle:UIAlertControllerStyleAlert];
+	UIAlertAction* okAction = [UIAlertAction
+									actionWithTitle:NSLocalizedString(@"OK", nil)
+									style:UIAlertActionStyleDefault
+									handler:^(UIAlertAction* action)
+									{
+										[AlertController dismissViewControllerAnimated : YES completion : nil];
+
+										UITextField* AlertTextField = AlertController.textFields.firstObject;
+										TextEntry = FText::FromString(AlertTextField.text);
+										AlertController = nil;
+										
+										FIOSAsyncTask* AsyncTask = [[FIOSAsyncTask alloc] init];
+										AsyncTask.GameThreadCallback = ^ bool(void)
 										{
-											[AlertController dismissViewControllerAnimated : YES completion : nil];
-
-											UITextField* AlertTextField = AlertController.textFields.firstObject;
-											TextEntry = FText::FromString(AlertTextField.text);
-                                            AlertController = nil;
-                                            
-											FIOSAsyncTask* AsyncTask = [[FIOSAsyncTask alloc] init];
-											AsyncTask.GameThreadCallback = ^ bool(void)
+											if(TextWidget.IsValid())
 											{
-                                                if(TextWidget.IsValid())
-                                                {
-                                                    TSharedPtr<IVirtualKeyboardEntry> TextEntryWidgetPin = TextWidget.Pin();
-                                                    TextEntryWidgetPin->SetTextFromVirtualKeyboard(TextEntry, ETextEntryType::TextEntryAccepted);
-                                                }
-
-												// clear the TextWidget
-												TextWidget = nullptr;
-												return true;
-											};
-											[AsyncTask FinishedTask];
-										}
-		];
-		UIAlertAction* cancelAction = [UIAlertAction
-										actionWithTitle: NSLocalizedString(@"Cancel", nil)
-										style:UIAlertActionStyleDefault
-										handler:^(UIAlertAction* action)
+												TSharedPtr<IVirtualKeyboardEntry> TextEntryWidgetPin = TextWidget.Pin();
+												TextEntryWidgetPin->SetTextFromVirtualKeyboard(TextEntry, ETextEntryType::TextEntryAccepted);
+											}
+
+											// clear the TextWidget
+											TextWidget = nullptr;
+											return true;
+										};
+										[AsyncTask FinishedTask];
+									}
+	];
+	UIAlertAction* cancelAction = [UIAlertAction
+									actionWithTitle: NSLocalizedString(@"Cancel", nil)
+									style:UIAlertActionStyleDefault
+									handler:^(UIAlertAction* action)
+									{
+										[AlertController dismissViewControllerAnimated : YES completion : nil];
+										AlertController = nil;
+										
+										FIOSAsyncTask* AsyncTask = [[FIOSAsyncTask alloc] init];
+										AsyncTask.GameThreadCallback = ^ bool(void)
 										{
-											[AlertController dismissViewControllerAnimated : YES completion : nil];
-                                            AlertController = nil;
-                                            
-											FIOSAsyncTask* AsyncTask = [[FIOSAsyncTask alloc] init];
-											AsyncTask.GameThreadCallback = ^ bool(void)
-											{
-												// clear the TextWidget
-												TextWidget = nullptr;
-												return true;
-											};
-											[AsyncTask FinishedTask];
-										}
-		];
-
-		[AlertController addAction: okAction];
-		[AlertController addAction: cancelAction];
-		[AlertController
-						addTextFieldWithConfigurationHandler:^(UITextField* AlertTextField)
+											// clear the TextWidget
+											TextWidget = nullptr;
+											return true;
+										};
+										[AsyncTask FinishedTask];
+									}
+	];
+
+	[AlertController addAction: okAction];
+	[AlertController addAction: cancelAction];
+	[AlertController
+					addTextFieldWithConfigurationHandler:^(UITextField* AlertTextField)
+					{
+						AlertTextField.clearsOnBeginEditing = NO;
+						AlertTextField.clearsOnInsertion = NO;
+						if (TextWidget.IsValid())
 						{
-							AlertTextField.clearsOnBeginEditing = NO;
-							AlertTextField.clearsOnInsertion = NO;
-                            if (TextWidget.IsValid())
-                            {
-                                TSharedPtr<IVirtualKeyboardEntry> TextEntryWidgetPin = TextWidget.Pin();
-                                AlertTextField.text = [NSString stringWithFString : TextEntryWidgetPin->GetText().ToString()];
-                                AlertTextField.placeholder = [NSString stringWithFString : TextEntryWidgetPin->GetHintText().ToString()];
-		 
-                                FKeyboardConfig KeyboardConfig;
-                                GetKeyboardConfig(TextEntryWidgetPin, KeyboardConfig);
-
-                                AlertTextField.keyboardType = KeyboardConfig.KeyboardType;
-                                AlertTextField.autocorrectionType = KeyboardConfig.AutocorrectionType;
-                                AlertTextField.autocapitalizationType = KeyboardConfig.AutocapitalizationType;
-                                AlertTextField.secureTextEntry = KeyboardConfig.bSecureTextEntry;
-                            }
+							AlertTextField.text = TextContents;
+							AlertTextField.placeholder = PlaceholderContents;
+							AlertTextField.keyboardType = KeyboardConfig.KeyboardType;
+							AlertTextField.autocorrectionType = KeyboardConfig.AutocorrectionType;
+							AlertTextField.autocapitalizationType = KeyboardConfig.AutocapitalizationType;
+							AlertTextField.secureTextEntry = KeyboardConfig.bSecureTextEntry;
 						}
-		];
-		[[IOSAppDelegate GetDelegate].IOSController presentViewController : AlertController animated : YES completion : nil];
-	}
-	else
-#endif
-	{
-#if __IPHONE_OS_VERSION_MIN_REQUIRED < __IPHONE_9_0
-		AlertView = [[UIAlertView alloc] initWithTitle:@""
-									message:@""
-									delegate:self
-									cancelButtonTitle:NSLocalizedString(@"Cancel", nil)
-									otherButtonTitles:NSLocalizedString(@"OK", nil), nil];
-
-		// give the UIAlertView a style so a UITextField is created
-		switch (TextWidget->GetVirtualKeyboardType())
-		{
-		case EKeyboardType::Keyboard_Password:
-			AlertView.alertViewStyle = UIAlertViewStyleSecureTextInput;
-			break;
-		case EKeyboardType::Keyboard_AlphaNumeric:
-		case EKeyboardType::Keyboard_Default:
-		case EKeyboardType::Keyboard_Email:
-		case EKeyboardType::Keyboard_Number:
-		case EKeyboardType::Keyboard_Web:
-		default:
-			AlertView.alertViewStyle = UIAlertViewStylePlainTextInput;
-			break;
-		}
-
-		UITextField* AlertTextField = [AlertView textFieldAtIndex : 0];
-        if (TextWidget.IsValid())
-        {
-            TSharedPtr<IVirtualKeyboardEntry> TextEntryWidgetPin = TextWidget.Pin();
-            AlertTextField.clearsOnBeginEditing = NO;
-            AlertTextField.clearsOnInsertion = NO;
-            AlertTextField.autocorrectionType = UITextAutocorrectionTypeNo;
-            AlertTextField.autocapitalizationType = UITextAutocapitalizationTypeNone;
-            AlertTextField.text = [NSString stringWithFString : TextEntryWidgetPin->GetText().ToString()];
-            AlertTextField.placeholder = [NSString stringWithFString : TextEntryWidgetPin->GetHintText().ToString()];
-
-            // set up the keyboard styles not supported in the AlertViewStyle styles
-            switch (TextEntryWidgetPin->GetVirtualKeyboardType())
-            {
-                case EKeyboardType::Keyboard_Email:
-                    AlertTextField.keyboardType = UIKeyboardTypeEmailAddress;
-                    break;
-                case EKeyboardType::Keyboard_Number:
-                    AlertTextField.keyboardType = UIKeyboardTypeDecimalPad;
-                    break;
-                case EKeyboardType::Keyboard_Web:
-                    AlertTextField.keyboardType = UIKeyboardTypeURL;
-                    break;
-                case EKeyboardType::Keyboard_AlphaNumeric:
-                    AlertTextField.keyboardType = UIKeyboardTypeASCIICapable;
-                    break;
-                case EKeyboardType::Keyboard_Default:
-                case EKeyboardType::Keyboard_Password:
-                default:
-                    // nothing to do, UIAlertView style handles these keyboard types
-                    break;
-            }
-        }
-		[AlertView show];
-
-		[AlertView release];
-#endif
-	}
+					}
+	];
+	[[IOSAppDelegate GetDelegate].IOSController presentViewController : AlertController animated : YES completion : nil];
 }
 
-#if __IPHONE_OS_VERSION_MIN_REQUIRED < __IPHONE_9_0
-- (void)alertView:(UIAlertView *)alertView clickedButtonAtIndex:(NSInteger)buttonIndex
-{
-	UITextField* AlertTextField = [alertView textFieldAtIndex : 0];
-	TextEntry = FText::FromString(AlertTextField.text);
-
-	FIOSAsyncTask* AsyncTask = [[FIOSAsyncTask alloc] init];
-    AsyncTask.GameThreadCallback = ^ bool(void)
-    {
-		// index 1 is the OK button
-		if(buttonIndex == 1 && TextWidget.IsValid())
-		{
-            TSharedPtr<IVirtualKeyboardEntry> TextEntryWidgetPin = TextWidget.Pin();
-			TextEntryWidgetPin->SetTextFromVirtualKeyboard(TextEntry, ETextEntryType::TextEntryAccepted);
-		}
-    
-        // clear the TextWidget
-        TextWidget = nullptr;
-        return true;
-    };
-    [AsyncTask FinishedTask];
-}
-#endif
-
 @end
 
 #endif
diff --git a/Engine/Source/Runtime/Slate/Public/Framework/Text/IOS/IOSPlatformTextField.h b/Engine/Source/Runtime/Slate/Public/Framework/Text/IOS/IOSPlatformTextField.h
--- a/Engine/Source/Runtime/Slate/Public/Framework/Text/IOS/IOSPlatformTextField.h
+++ b/Engine/Source/Runtime/Slate/Public/Framework/Text/IOS/IOSPlatformTextField.h
@@ -5,6 +5,7 @@
 #include "CoreMinimal.h"
 #include "Framework/Application/IPlatformTextField.h"
 #include "Internationalization/Text.h"
+#include "IOSView.h"
 
 #import <UIKit/UIKit.h>
 
@@ -41,7 +42,7 @@ typedef FIOSPlatformTextField FPlatformTextField;
 #endif
 }
 
--(void)show:(TSharedPtr<IVirtualKeyboardEntry>)InTextWidget;
+-(void)show:(TSharedPtr<IVirtualKeyboardEntry>)InTextWidget text:(NSString*)TextContents placeholder:(NSString*)PlaceholderContents keyboardConfig:(FKeyboardConfig)KeyboardConfig;
 -(void)hide;
 -(bool)hasTextWidget;
 
diff --git a/Engine/Source/Runtime/WebBrowser/Private/WebJSFunction.cpp b/Engine/Source/Runtime/WebBrowser/Private/WebJSFunction.cpp
--- a/Engine/Source/Runtime/WebBrowser/Private/WebJSFunction.cpp
+++ b/Engine/Source/Runtime/WebBrowser/Private/WebJSFunction.cpp
@@ -83,6 +83,47 @@ FWebJSParam::FWebJSParam(const FWebJSParam& Other)
 	}
 }
 
+FWebJSParam::FWebJSParam(FWebJSParam&& Other)
+	: Tag(Other.Tag)
+{
+	switch (Other.Tag)
+	{
+	case PTYPE_BOOL:
+		BoolValue = Other.BoolValue;
+		break;
+	case PTYPE_DOUBLE:
+		DoubleValue = Other.DoubleValue;
+		break;
+	case PTYPE_INT:
+		IntValue = Other.IntValue;
+		break;
+	case PTYPE_STRING:
+		StringValue = Other.StringValue;
+		Other.StringValue = nullptr;
+		break;
+	case PTYPE_NULL:
+		break;
+	case PTYPE_OBJECT:
+		ObjectValue = Other.ObjectValue;
+		Other.ObjectValue = nullptr;
+		break;
+	case PTYPE_STRUCT:
+		StructValue = Other.StructValue;
+		Other.StructValue = nullptr;
+		break;
+	case PTYPE_ARRAY:
+		ArrayValue = Other.ArrayValue;
+		Other.ArrayValue = nullptr;
+		break;
+	case PTYPE_MAP:
+		MapValue = Other.MapValue;
+		Other.MapValue = nullptr;
+		break;
+	}
+
+	Other.Tag = PTYPE_NULL;
+}
+
 void FWebJSCallbackBase::Invoke(int32 ArgCount, FWebJSParam Arguments[], bool bIsError) const
 {
 	TSharedPtr<FWebJSScripting> Scripting = ScriptingPtr.Pin();
diff --git a/Engine/Source/Runtime/WebBrowser/Public/WebJSFunction.h b/Engine/Source/Runtime/WebBrowser/Public/WebJSFunction.h
--- a/Engine/Source/Runtime/WebBrowser/Public/WebJSFunction.h
+++ b/Engine/Source/Runtime/WebBrowser/Public/WebJSFunction.h
@@ -97,6 +97,7 @@ struct WEBBROWSER_API FWebJSParam
 		}
 	}
 	FWebJSParam(const FWebJSParam& Other);
+	FWebJSParam(FWebJSParam&& Other);
 	~FWebJSParam();
 
 	enum { PTYPE_NULL, PTYPE_BOOL, PTYPE_INT, PTYPE_DOUBLE, PTYPE_STRING, PTYPE_OBJECT, PTYPE_STRUCT, PTYPE_ARRAY, PTYPE_MAP } Tag;
