Ask for external storage permission in SplashActivity instead of GameActivity

#ue4
#android

diff --git a/Engine/Build/Android/Java/src/com/epicgames/ue4/SplashActivity.java b/Engine/Build/Android/Java/src/com/epicgames/ue4/SplashActivity.java
--- a/Engine/Build/Android/Java/src/com/epicgames/ue4/SplashActivity.java
+++ b/Engine/Build/Android/Java/src/com/epicgames/ue4/SplashActivity.java
@@ -7,13 +7,23 @@ import android.os.Bundle;
 import android.app.Activity;
 import android.content.Intent;
 import android.content.pm.ApplicationInfo;
+import android.content.pm.PackageInfo;
 import android.content.pm.PackageManager;
 import android.content.pm.PackageManager.NameNotFoundException;
 import android.view.View;
 import android.view.WindowManager;
 
+import android.support.v4.app.ActivityCompat;
+import android.support.v4.content.ContextCompat;
+
 public class SplashActivity extends Activity
 {
+	private static final int PERMISSION_REQUEST_CODE = 1105;
+
+	private Intent GameActivityIntent;
+	private boolean WaitForPermission = false;
+
+
 	public static Logger Log = new Logger("UE4-SplashActivity");
 	
 	@SuppressLint("ObsoleteSdkInt")
@@ -24,6 +34,8 @@ public class SplashActivity extends Activity
 
 		boolean ShouldHideUI = false;
 		boolean UseDisplayCutout = false;
+		boolean IsShipping = false;
+		boolean UseExternalFilesDir = false;
 		try {
 			ApplicationInfo ai = getPackageManager().getApplicationInfo(getPackageName(), PackageManager.GET_META_DATA);
 			Bundle bundle = ai.metaData;
@@ -36,6 +48,15 @@ public class SplashActivity extends Activity
 			{
 				UseDisplayCutout = bundle.getBoolean("com.epicgames.ue4.GameActivity.bUseDisplayCutout");
 			}
+			if (bundle.containsKey("com.epicgames.ue4.GameActivity.BuildConfiguration"))
+			{
+				String Configuration = bundle.getString("com.epicgames.ue4.GameActivity.BuildConfiguration");
+				IsShipping = Configuration.equals("Shipping");
+			}
+			if (bundle.containsKey("com.epicgames.ue4.GameActivity.bUseExternalFilesDir"))
+            {
+                UseExternalFilesDir = bundle.getBoolean("com.epicgames.ue4.GameActivity.bUseExternalFilesDir");
+            }
 		}
 		catch (NameNotFoundException | NullPointerException e)
 		{
@@ -77,17 +98,17 @@ public class SplashActivity extends Activity
 			}
 		}
 
-		Intent intent = new Intent(this, GameActivity.class);
-		intent.putExtras(getIntent());
-		intent.addFlags(Intent.FLAG_ACTIVITY_NO_ANIMATION);
-		intent.putExtra("UseSplashScreen", true);
+		GameActivityIntent = new Intent(this, GameActivity.class);
+		GameActivityIntent.putExtras(getIntent());
+		GameActivityIntent.addFlags(Intent.FLAG_ACTIVITY_NO_ANIMATION);
+		GameActivityIntent.putExtra("UseSplashScreen", "true");
 		if (ShouldHideUI)
 		{
-			intent.putExtra("ShouldHideUI", true);
+			GameActivityIntent.putExtra("ShouldHideUI", "true");
 		}
 		if (UseDisplayCutout)
 		{
-			intent.putExtra("UseDisplayCutout", "true");
+			GameActivityIntent.putExtra("UseDisplayCutout", "true");
 		}
 
 		//pass down any extras added to this Activity's intent to the GameActivity intent (GCM data, for example)
@@ -95,27 +116,75 @@ public class SplashActivity extends Activity
 		Bundle intentBundle = intentFromActivity.getExtras();
 		if(intentBundle != null)
 		{
-			intent.putExtras(intentBundle);
+			GameActivityIntent.putExtras(intentBundle);
 		}
 		
 		// pass the action if available
 		String intentAction = intentFromActivity.getAction();
 		if (intentAction != null)
 		{
-			intent.setAction(intentAction);
+			GameActivityIntent.setAction(intentAction);
+		}
+
+		// check if we need to wait for permissions
+		int targetSdkVersion = 0;
+		try 
+		{
+			PackageInfo packageInfo = getPackageManager().getPackageInfo(getPackageName(), 0);
+			targetSdkVersion = packageInfo.applicationInfo.targetSdkVersion;
+		}
+		catch (PackageManager.NameNotFoundException e) 
+		{
+		}
+
+		if (android.os.Build.VERSION.SDK_INT >= 23 && targetSdkVersion >= 23) //23 is the API level (Marshmallow) where runtime permission handling is available
+		{
+			// we might need to ask for permission if we don't already have it
+			if (ContextCompat.checkSelfPermission(this, "android.permission.WRITE_EXTERNAL_STORAGE") != PackageManager.PERMISSION_GRANTED)
+			{
+				if (!IsShipping || !UseExternalFilesDir)
+				{
+					ActivityCompat.requestPermissions(this, new String[] {"android.permission.WRITE_EXTERNAL_STORAGE"}, PERMISSION_REQUEST_CODE);
+					WaitForPermission = true;
+				}
+			}
 		}
 
-		startActivity(intent);
-		finish();
-		overridePendingTransition(0, 0);
+		if (!WaitForPermission)
+		{
+			startActivity(GameActivityIntent);
+			finish();
+			overridePendingTransition(0, 0);
+		}
+	}
+
+	@Override
+	public void onRequestPermissionsResult(int requestCode, String[] permissions, int[] grantResults)
+	{
+		if (requestCode==PERMISSION_REQUEST_CODE && permissions.length>0) 
+		{
+			if (grantResults.length>0 && grantResults[0]==PackageManager.PERMISSION_GRANTED)
+			{
+				startActivity(GameActivityIntent);
+				finish();
+				overridePendingTransition(0, 0);
+			}
+			else
+			{
+				finish();
+			}
+		}
 	}
 
 	@Override
 	protected void onPause()
 	{
 		super.onPause();
-		finish();
-		overridePendingTransition(0, 0);
+		if (!WaitForPermission)
+		{
+			finish();
+			overridePendingTransition(0, 0);
+		}
 	}
 
 }
\ No newline at end of file
