Fix intent extra data passing to be consistent between SplashActivity and GameActivity

#ue4
#android

diff --git a/Engine/Build/Android/Java/src/com/epicgames/ue4/GameActivity.java.template b/Engine/Build/Android/Java/src/com/epicgames/ue4/GameActivity.java.template
--- a/Engine/Build/Android/Java/src/com/epicgames/ue4/GameActivity.java.template
+++ b/Engine/Build/Android/Java/src/com/epicgames/ue4/GameActivity.java.template
@@ -2110,9 +2110,9 @@ public class GameActivity extends $${gameActivitySuperClass}$$ implements Surfac
 		_extrasBundle = getIntent().getExtras();
 		if (_extrasBundle != null)
 		{
-			ShouldHideUI = _extrasBundle.getBoolean("ShouldHideUI", false);
-			UseDisplayCutout = _extrasBundle.getBoolean("UseDisplayCutout", false);
-			if (_extrasBundle.getBoolean("UseSplashScreen", false))
+			ShouldHideUI = _extrasBundle.getString("ShouldHideUI") != null;
+			UseDisplayCutout = _extrasBundle.getString("UseDisplayCutout") != null;
+			if (_extrasBundle.getString("UseSplashScreen") != null)
 			{
 				SplashScreenLaunch = true;
 				try {
