Allow fullscreen on S10

#android

diff --git a/Engine/Build/Android/Java/src/com/epicgames/ue4/SplashActivity.java b/Engine/Build/Android/Java/src/com/epicgames/ue4/SplashActivity.java
--- a/Engine/Build/Android/Java/src/com/epicgames/ue4/SplashActivity.java
+++ b/Engine/Build/Android/Java/src/com/epicgames/ue4/SplashActivity.java
@@ -77,8 +77,24 @@ public class SplashActivity extends Activity
 			}
 		}
 
-		// for now only allow on one device manufacturer - fix up later
-		if (!android.os.Build.MANUFACTURER.equals("HUAWEI"))
+		
+		// allow certain models for now to use full area around cutout
+		boolean BlockDisplayCutout = true;
+		if (android.os.Build.MANUFACTURER.equals("HUAWEI"))
+		{
+			BlockDisplayCutout = false;
+		}
+		else if (android.os.Build.MANUFACTURER.equals("samsung"))
+		{
+			String model = android.os.Build.MODEL;
+			if (model.startsWith("SM-G970") || model.startsWith("SM-G973") || model.startsWith("SM-G975") ||
+				model.startsWith("SC-03L") || model.startsWith("SCV41") || model.startsWith("SC-04L") || model.startsWith("SCV42"))
+			{
+				BlockDisplayCutout = false;
+			}
+		}
+			
+		if (BlockDisplayCutout)
 		{
 			UseDisplayCutout = false;
 		}
