Fix local notifications for Android

#android

diff --git a/Engine/Build/Android/Java/src/com/epicgames/ue4/GameActivity.java.template b/Engine/Build/Android/Java/src/com/epicgames/ue4/GameActivity.java.template
--- a/Engine/Build/Android/Java/src/com/epicgames/ue4/GameActivity.java.template
+++ b/Engine/Build/Android/Java/src/com/epicgames/ue4/GameActivity.java.template
@@ -4613,9 +4613,6 @@ public class GameActivity extends $${gameActivitySuperClass}$$ implements Surfac
 			}
 			editor.putString("notificationIDs", notificationIDs + "-" + idToReturn);
 		}
-
-		notificationIDs = preferences.getString("notificationIDs", "");
-
 		editor.commit();
 
 		return idToReturn;
@@ -4640,11 +4637,30 @@ public class GameActivity extends $${gameActivitySuperClass}$$ implements Surfac
 		return iParts;
 	}
 
+	private boolean LocalNotificationIDExists(int notificationID)
+	{
+		if (notificationID == -1)
+		{
+			return false;
+		}
+
+		ArrayList<Integer> notificationIDs = LocalNotificationGetIDList();
+
+		for(int id : notificationIDs)
+		{
+			if (id == notificationID)
+			{
+				return true;
+			}
+		}
+		return false;
+	}
+
 	private void LocalNotificationRemoveID(int notificationID)
 	{
 		SharedPreferences preferences = getApplicationContext().getSharedPreferences("LocalNotificationPreferences", MODE_PRIVATE);
 		SharedPreferences.Editor editor = preferences.edit();
-		String notificationIDs = preferences.getString("notificationIDs", null);
+		String notificationIDs = preferences.getString("notificationIDs", "");
 
 		ArrayList<String> iParts = new ArrayList<String>();
 
@@ -4682,7 +4698,7 @@ public class GameActivity extends $${gameActivitySuperClass}$$ implements Surfac
 		editor.commit();
 	}
 
-	public void AndroidThunkJava_LocalNotificationScheduleAtTime(String targetDateTime, boolean localTime, String title, String body, String action, String activationEvent) 
+	public int AndroidThunkJava_LocalNotificationScheduleAtTime(String targetDateTime, boolean localTime, String title, String body, String action, String activationEvent) 
 	{
 		int notificationID = LocalNotificationGetID(this);
 
@@ -4719,7 +4735,8 @@ public class GameActivity extends $${gameActivitySuperClass}$$ implements Surfac
 		catch (ParseException e) 
 		{
 			e.printStackTrace();
-			return;
+			LocalNotificationRemoveID(notificationID);
+			return -1;
 		}
 
 		Date currentDate = new Date();
@@ -4728,7 +4745,8 @@ public class GameActivity extends $${gameActivitySuperClass}$$ implements Surfac
 
 		if(msDiff < 0)
 		{
-			return;
+			LocalNotificationRemoveID(notificationID);
+			return -1;
 		}
 
 		long futureTimeInMillis = SystemClock.elapsedRealtime() + msDiff;//Calculate the time to run the callback
@@ -4736,6 +4754,8 @@ public class GameActivity extends $${gameActivitySuperClass}$$ implements Surfac
 
 		//Schedule the operation by using AlarmService
 		alarmManager.set(AlarmManager.ELAPSED_REALTIME_WAKEUP, futureTimeInMillis, pendingIntent);
+
+		return notificationID;
 	}
 
 	public class LaunchNotification {
@@ -4767,16 +4787,29 @@ public class GameActivity extends $${gameActivitySuperClass}$$ implements Surfac
 			pendingIntent.cancel();
 			alarmManager.cancel(pendingIntent);
 		}
+
+		// clear them all
+		SharedPreferences preferences = getApplicationContext().getSharedPreferences("LocalNotificationPreferences", MODE_PRIVATE);
+		SharedPreferences.Editor editor = preferences.edit();
+		editor.putString("notificationIDs", "");
+		editor.commit();
+	}
+
+	public boolean AndroidThunkJava_LocalNotificationExists(int notificationId)
+	{
+		return LocalNotificationIDExists(notificationId);
 	}
 
 	// Returns true only if the scheduled notification exists and gets destoyed successfully
 	public boolean AndroidThunkJava_LocalNotificationDestroyIfExists(int notificationId)
 	{
-		if (AndroidThunkJava_ScheduledNotificationExists(notificationId))
+		if (AndroidThunkJava_LocalNotificationExists(notificationId))
 		{
+			LocalNotificationRemoveID(notificationId);
+
 			//Cancel the intent itself as well as from the alarm manager
 			AlarmManager alarmManager = (AlarmManager)getSystemService(Context.ALARM_SERVICE);
-			PendingIntent pendingIntent = PendingIntent.getBroadcast(this, notificationId, new Intent(this, ScheduledNotificationReceiver.class), PendingIntent.FLAG_UPDATE_CURRENT);
+			PendingIntent pendingIntent = PendingIntent.getBroadcast(this, notificationId, new Intent(this, LocalNotificationReceiver.class), PendingIntent.FLAG_UPDATE_CURRENT);
 			pendingIntent.cancel();
 			alarmManager.cancel(pendingIntent);
 			return true;
diff --git a/Engine/Source/Runtime/Launch/Private/Android/AndroidJNI.cpp b/Engine/Source/Runtime/Launch/Private/Android/AndroidJNI.cpp
--- a/Engine/Source/Runtime/Launch/Private/Android/AndroidJNI.cpp
+++ b/Engine/Source/Runtime/Launch/Private/Android/AndroidJNI.cpp
@@ -98,8 +98,9 @@ void FJavaWrapper::FindClassesAndMethods(JNIEnv* Env)
 	AndroidThunkJava_GetMetaDataString = FindMethod(Env, GameActivityClassID, "AndroidThunkJava_GetMetaDataString", "(Ljava/lang/String;)Ljava/lang/String;", bIsOptional);
 	AndroidThunkJava_SetSustainedPerformanceMode = FindMethod(Env, GameActivityClassID, "AndroidThunkJava_SetSustainedPerformanceMode", "(Z)V", bIsOptional);
 	AndroidThunkJava_ShowHiddenAlertDialog = FindMethod(Env, GameActivityClassID, "AndroidThunkJava_ShowHiddenAlertDialog", "()V", bIsOptional);
-	AndroidThunkJava_LocalNotificationScheduleAtTime = FindMethod(Env, GameActivityClassID, "AndroidThunkJava_LocalNotificationScheduleAtTime", "(Ljava/lang/String;ZLjava/lang/String;Ljava/lang/String;Ljava/lang/String;Ljava/lang/String;)V", bIsOptional);
+	AndroidThunkJava_LocalNotificationScheduleAtTime = FindMethod(Env, GameActivityClassID, "AndroidThunkJava_LocalNotificationScheduleAtTime", "(Ljava/lang/String;ZLjava/lang/String;Ljava/lang/String;Ljava/lang/String;Ljava/lang/String;)I", bIsOptional);
 	AndroidThunkJava_LocalNotificationClearAll = FindMethod(Env, GameActivityClassID, "AndroidThunkJava_LocalNotificationClearAll", "()V", bIsOptional);
+	AndroidThunkJava_LocalNotificationExists = FindMethod(Env, GameActivityClassID, "AndroidThunkJava_LocalNotificationExists", "(I)Z", bIsOptional);
 	AndroidThunkJava_LocalNotificationGetLaunchNotification = FindMethod(Env, GameActivityClassID, "AndroidThunkJava_LocalNotificationGetLaunchNotification", "()Lcom/epicgames/ue4/GameActivity$LaunchNotification;", bIsOptional);
 	AndroidThunkJava_LocalNotificationDestroyIfExists = FindMethod(Env, GameActivityClassID, "AndroidThunkJava_LocalNotificationDestroyIfExists", "(I)Z", bIsOptional);
 	AndroidThunkJava_GetNetworkConnectionType = FindMethod(Env, GameActivityClassID, "AndroidThunkJava_GetNetworkConnectionType", "()I", bIsOptional);
@@ -377,6 +378,7 @@ jmethodID FJavaWrapper::AndroidThunkJava_UnregisterForRemoteNotifications;
 jmethodID FJavaWrapper::AndroidThunkJava_ShowHiddenAlertDialog;
 jmethodID FJavaWrapper::AndroidThunkJava_LocalNotificationScheduleAtTime;
 jmethodID FJavaWrapper::AndroidThunkJava_LocalNotificationClearAll;
+jmethodID FJavaWrapper::AndroidThunkJava_LocalNotificationExists;
 jmethodID FJavaWrapper::AndroidThunkJava_LocalNotificationGetLaunchNotification;
 jmethodID FJavaWrapper::AndroidThunkJava_LocalNotificationDestroyIfExists;
 jmethodID FJavaWrapper::AndroidThunkJava_GetNetworkConnectionType;
@@ -1468,6 +1470,16 @@ void AndroidThunkCpp_ClearAllLocalNotifications()
 	}
 }
 
+bool AndroidThunkCpp_LocalNotificationExists(int32 NotificationId)
+{
+	bool Result = false;
+	if (JNIEnv* Env = FAndroidApplication::GetJavaEnv())
+	{
+		Result = FJavaWrapper::CallBooleanMethod(Env, FJavaWrapper::GameActivityThis, FJavaWrapper::AndroidThunkJava_LocalNotificationExists, NotificationId);
+	}
+	return Result;
+}
+
 bool AndroidThunkCpp_DestroyScheduledNotificationIfExists(int32 NotificationId)
 {
 	bool Result = false;
diff --git a/Engine/Source/Runtime/Launch/Public/Android/AndroidJNI.h b/Engine/Source/Runtime/Launch/Public/Android/AndroidJNI.h
--- a/Engine/Source/Runtime/Launch/Public/Android/AndroidJNI.h
+++ b/Engine/Source/Runtime/Launch/Public/Android/AndroidJNI.h
@@ -56,6 +56,7 @@ public:
 	static jmethodID AndroidThunkJava_ShowHiddenAlertDialog;
 	static jmethodID AndroidThunkJava_LocalNotificationScheduleAtTime;
 	static jmethodID AndroidThunkJava_LocalNotificationClearAll;
+	static jmethodID AndroidThunkJava_LocalNotificationExists;
 	static jmethodID AndroidThunkJava_LocalNotificationGetLaunchNotification;
 	static jmethodID AndroidThunkJava_LocalNotificationDestroyIfExists;
 	static jmethodID AndroidThunkJava_GetNetworkConnectionType;
