UE-74327 Unable to summon the virtual keyboard on Android 9.0 with TargetSDKVersion=28

#Android
#4.23

Manually request focus when clicking on a EditText (https://developer.android.com/about/versions/pie/android-9.0-changes-28#focus)

diff --git a/Engine/Build/Android/Java/src/com/epicgames/ue4/GameActivity.java.template b/Engine/Build/Android/Java/src/com/epicgames/ue4/GameActivity.java.template
--- a/Engine/Build/Android/Java/src/com/epicgames/ue4/GameActivity.java.template
+++ b/Engine/Build/Android/Java/src/com/epicgames/ue4/GameActivity.java.template
@@ -2652,6 +2652,7 @@ public class GameActivity extends $${gameActivitySuperClass}$$ implements Surfac
 						// remember down time and position
 						downTime = System.currentTimeMillis();
 						downX = event.getX();
+						consoleInputBox.requestFocus();
 						return true;
 					}
 					case MotionEvent.ACTION_UP: {
