Fix issue with missing class definition crash on KitKat

#ue4
#android

diff --git a/Engine/Build/Android/Java/src/com/epicgames/ue4/GameActivity.java.template b/Engine/Build/Android/Java/src/com/epicgames/ue4/GameActivity.java.template
--- a/Engine/Build/Android/Java/src/com/epicgames/ue4/GameActivity.java.template
+++ b/Engine/Build/Android/Java/src/com/epicgames/ue4/GameActivity.java.template
@@ -2843,46 +2843,49 @@ public class GameActivity extends $${gameActivitySuperClass}$$ implements Surfac
 		
 		// check for display cutouts to set safe zone
 		View mainRootView = mainDecorView.findViewById(android.R.id.content).getRootView();
-		mainRootView.setOnApplyWindowInsetsListener(new View.OnApplyWindowInsetsListener() {
-            @Override
-       	    public WindowInsets onApplyWindowInsets(View v, WindowInsets insets) {
-				if (Build.VERSION.SDK_INT >= 28)
-				{
-					DisplayCutout cutout = insets.getDisplayCutout();
-					if (cutout != null)
+		if (Build.VERSION.SDK_INT > 19)
+		{
+			mainRootView.setOnApplyWindowInsetsListener(new View.OnApplyWindowInsetsListener() {
+				@Override
+       			public WindowInsets onApplyWindowInsets(View v, WindowInsets insets) {
+					if (Build.VERSION.SDK_INT >= 28)
 					{
-						Rect fullViewRect = new Rect();
-						mainDecorView.getDrawingRect( mainDecorViewRect );
-
-						android.view.Display display = getWindowManager().getDefaultDisplay();
-						android.graphics.Point displaySize = new android.graphics.Point();
-						display.getRealSize(displaySize);
-
-						// store safezone independent of resolution
-						float invX = 1.0f / (displaySize.x != 0.0f ? displaySize.x : 1.0f);
-						float invY = 1.0f / (displaySize.y != 0.0f ? displaySize.y : 1.0f);
-						safezoneF.left = (safezone.left = cutout.getSafeInsetLeft()) * invX;
-						safezoneF.top = (safezone.top = cutout.getSafeInsetTop()) * invY;
-						safezoneF.right = (safezone.right = cutout.getSafeInsetRight()) * invX;
-						safezoneF.bottom = (safezone.bottom = cutout.getSafeInsetBottom()) *invY;
-						nativeSetSafezoneInfo(displaySize.y > displaySize.x, safezoneF.left, safezoneF.top, safezoneF.right, safezoneF.bottom);
-
-						//Log.debug("FullViewRect: " + fullViewRect.toString());	
-						//Log.debug("DisplayCutout: " + cutout.toString());
-						//Log.debug("SafeTop: " + cutout.getSafeInsetTop());
-						//Log.debug("SafeBottom: " + cutout.getSafeInsetBottom());
-						//Log.debug("SafeLeft: " + cutout.getSafeInsetLeft());
-						//Log.debug("SafeRight: " + cutout.getSafeInsetRight());
-						//List<Rect>bounds = cutout.getBoundingRects();
-						//for (Rect r : bounds)
-						//{
-						//	Log.debug("Rect: " + r.toString());
-						//}
+						DisplayCutout cutout = insets.getDisplayCutout();
+						if (cutout != null)
+						{
+							Rect fullViewRect = new Rect();
+							mainDecorView.getDrawingRect( mainDecorViewRect );
+
+							android.view.Display display = getWindowManager().getDefaultDisplay();
+							android.graphics.Point displaySize = new android.graphics.Point();
+							display.getRealSize(displaySize);
+
+							// store safezone independent of resolution
+							float invX = 1.0f / (displaySize.x != 0.0f ? displaySize.x : 1.0f);
+							float invY = 1.0f / (displaySize.y != 0.0f ? displaySize.y : 1.0f);
+							safezoneF.left = (safezone.left = cutout.getSafeInsetLeft()) * invX;
+							safezoneF.top = (safezone.top = cutout.getSafeInsetTop()) * invY;
+							safezoneF.right = (safezone.right = cutout.getSafeInsetRight()) * invX;
+							safezoneF.bottom = (safezone.bottom = cutout.getSafeInsetBottom()) *invY;
+							nativeSetSafezoneInfo(displaySize.y > displaySize.x, safezoneF.left, safezoneF.top, safezoneF.right, safezoneF.bottom);
+
+							//Log.debug("FullViewRect: " + fullViewRect.toString());	
+							//Log.debug("DisplayCutout: " + cutout.toString());
+							//Log.debug("SafeTop: " + cutout.getSafeInsetTop());
+							//Log.debug("SafeBottom: " + cutout.getSafeInsetBottom());
+							//Log.debug("SafeLeft: " + cutout.getSafeInsetLeft());
+							//Log.debug("SafeRight: " + cutout.getSafeInsetRight());
+							//List<Rect>bounds = cutout.getBoundingRects();
+							//for (Rect r : bounds)
+							//{
+							//	Log.debug("Rect: " + r.toString());
+							//}
+						}
 					}
+					return insets;
 				}
-				return insets;
-			}
-        });
+			});
+		}
 
 		createVirtualKeyboardInput();
 		
