Fix build issue for Android

#android
#kairos

diff --git a/Engine/Build/Android/Java/src/com/epicgames/ue4/GameActivity.java.template b/Engine/Build/Android/Java/src/com/epicgames/ue4/GameActivity.java.template
--- a/Engine/Build/Android/Java/src/com/epicgames/ue4/GameActivity.java.template
+++ b/Engine/Build/Android/Java/src/com/epicgames/ue4/GameActivity.java.template
@@ -141,6 +141,7 @@ import android.view.WindowManager;
 import android.view.Window;
 import android.widget.LinearLayout;
 import android.widget.PopupWindow;
+import android.text.TextUtils;
 
 import android.media.AudioManager;
 
