Don't try to use NetworkRequest.Builder on KitKat

#ue4
#android

diff --git a/Engine/Build/Android/Java/src/com/epicgames/ue4/network/NetworkChangedManager.java b/Engine/Build/Android/Java/src/com/epicgames/ue4/network/NetworkChangedManager.java
--- a/Engine/Build/Android/Java/src/com/epicgames/ue4/network/NetworkChangedManager.java
+++ b/Engine/Build/Android/Java/src/com/epicgames/ue4/network/NetworkChangedManager.java
@@ -36,6 +36,12 @@ public final class NetworkChangedManager {
 	}
 
 	public void initNetworkCallback(@NonNull Context context) {
+		if (android.os.Build.VERSION.SDK_INT < 21)
+		{
+			// unsupported before Lollipop
+			return;
+		}
+)
 		ConnectivityManager connectivityManager = (ConnectivityManager) context.getSystemService(Context.CONNECTIVITY_SERVICE);
 		if (connectivityManager != null) {
 			NetworkRequest.Builder builder = new NetworkRequest.Builder();
