Fix typo

#android

diff --git a/Engine/Build/Android/Java/src/com/epicgames/ue4/network/NetworkChangedManager.java b/Engine/Build/Android/Java/src/com/epicgames/ue4/network/NetworkChangedManager.java
--- a/Engine/Build/Android/Java/src/com/epicgames/ue4/network/NetworkChangedManager.java
+++ b/Engine/Build/Android/Java/src/com/epicgames/ue4/network/NetworkChangedManager.java
@@ -41,7 +41,7 @@ public final class NetworkChangedManager {
 			// unsupported before Lollipop
 			return;
 		}
-)
+
 		ConnectivityManager connectivityManager = (ConnectivityManager) context.getSystemService(Context.CONNECTIVITY_SERVICE);
 		if (connectivityManager != null) {
 			NetworkRequest.Builder builder = new NetworkRequest.Builder();
