Make sure Android native input box gets hidden when we dismiss the keyboard
#rnx

diff --git a/Engine/Build/Android/Java/src/com/epicgames/ue4/GameActivity.java.template b/Engine/Build/Android/Java/src/com/epicgames/ue4/GameActivity.java.template
--- a/Engine/Build/Android/Java/src/com/epicgames/ue4/GameActivity.java.template
+++ b/Engine/Build/Android/Java/src/com/epicgames/ue4/GameActivity.java.template
@@ -3897,7 +3897,7 @@ public class GameActivity extends $${gameActivitySuperClass}$$ implements Surfac
 				break;
 				case VK_CMD_HIDE:
 				{
-						if(bKeyboardShowing)
+						if (bKeyboardShowing || newVirtualKeyboardInput.getVisibility() != View.GONE)
 						{
 							//Log.debug("VK: Hide newVirtualKeyboardInput");
 
