Add logging to gameactivity to help understand why some devices do not properly gather GPU information.
#rnx

diff --git a/Engine/Build/Android/Java/src/com/epicgames/ue4/GameActivity.java.template b/Engine/Build/Android/Java/src/com/epicgames/ue4/GameActivity.java.template
--- a/Engine/Build/Android/Java/src/com/epicgames/ue4/GameActivity.java.template
+++ b/Engine/Build/Android/Java/src/com/epicgames/ue4/GameActivity.java.template
@@ -413,6 +413,7 @@ public class GameActivity extends $${gameActivitySuperClass}$$ implements Surfac
 	private boolean InitCompletedOK = false;
 
 	private boolean bConfigRulesError = false;
+	private String ConfigRulesStatus = "UNSET";
 
 	/** Dialog for force game dialog */
 	private AlertDialog ForceGameDialog = null;
@@ -468,6 +469,7 @@ public class GameActivity extends $${gameActivitySuperClass}$$ implements Surfac
 	public String OpenGLDriver = "";
 	public String TextureFormats = "ETC1";
 	public boolean bSupportsFloatingPointRenderTargets = false;
+	private String EGLErrorStatus = "UNSET";
 
 	/** Used for LocalNotification support*/
 	private boolean localNotificationAppLaunched = false;
@@ -942,6 +944,8 @@ public class GameActivity extends $${gameActivitySuperClass}$$ implements Surfac
 	// returns InputStream to most up to date version of configrules
 	private InputStream getConfigRulesStream(String ProjectName)
 	{
+		ConfigRulesStatus = "";
+
 		byte[] configBytesAssets = getByteArrayFromAssets("configrules.bin.png");
 		byte[] configBytesInternal = getByteArrayFromFile(InternalFilesDir + "configrules.bin.png");
 		byte[] configBytesExternal = getByteArrayFromFile(android.os.Environment.getExternalStorageDirectory().getAbsolutePath() + "/UE4Game/" + ProjectName + "/configrules.bin.png");
@@ -953,6 +957,7 @@ public class GameActivity extends $${gameActivitySuperClass}$$ implements Surfac
 		// is there one at all?
 		if (assetsVersion == -1 && internalVersion == -1 && externalVersion == -1)
 		{
+			ConfigRulesStatus = "Not found. ";
 			return null;
 		}
 
@@ -987,6 +992,7 @@ public class GameActivity extends $${gameActivitySuperClass}$$ implements Surfac
 		}
 		if (result == null)
 		{
+			ConfigRulesStatus = "failed to find config rules. ";
 			bConfigRulesError = true;
 		}
 		return result;
@@ -1117,9 +1123,15 @@ public class GameActivity extends $${gameActivitySuperClass}$$ implements Surfac
 
 	private void getGraphicsInformation()
     {
+		int initialError = EGL14.eglGetError();
+		EGLErrorStatus = (initialError == EGL14.EGL_SUCCESS) ? "" : "InitialEGLError = " + initialError + " ";
+
 		EGLDisplay display = EGL14.eglGetDisplay(EGL14.EGL_DEFAULT_DISPLAY);
+		EGLErrorStatus += (display != EGL14.EGL_NO_DISPLAY) ? "" : "eglGetDisplay = " + EGL14.eglGetError() + " ";
+
 		int[] version = new int[2];
-		EGL14.eglInitialize(display, version, 0, version, 1);
+		boolean bEglInitialized = EGL14.eglInitialize(display, version, 0, version, 1);
+		EGLErrorStatus += (bEglInitialized) ? "" : "eglInitialize = " + EGL14.eglGetError() + " ";
 
 		// find a config from EGL
 		int[] configAttr = {EGL14.EGL_COLOR_BUFFER_TYPE, EGL14.EGL_RGB_BUFFER,
@@ -1132,6 +1144,7 @@ public class GameActivity extends $${gameActivitySuperClass}$$ implements Surfac
 		EGL14.eglChooseConfig(display, configAttr, 0, configs, 0, 1, numConfig, 0);
 		if (numConfig[0] == 0)
 		{
+			EGLErrorStatus += "eglChooseConfig = " + EGL14.eglGetError() + " ";
 			Log.debug("Did not find an EGL config");
 			return;
 		}
@@ -1140,13 +1153,16 @@ public class GameActivity extends $${gameActivitySuperClass}$$ implements Surfac
 		// small surface just to make a context
 		int[] surfAttr = {EGL14.EGL_WIDTH, 8, EGL14.EGL_HEIGHT, 8, EGL14.EGL_NONE};
 		EGLSurface surface = EGL14.eglCreatePbufferSurface(display, config, surfAttr, 0);
+		EGLErrorStatus += (surface != EGL14.EGL_NO_SURFACE) ? "" : "eglCreatePbufferSurface = " + EGL14.eglGetError() + " ";
 
 		// create the context (ES 2)
 		int[] contextAttr = {EGL14.EGL_CONTEXT_CLIENT_VERSION, 2, EGL14.EGL_NONE};
 		EGLContext context = EGL14.eglCreateContext(display, config, EGL14.EGL_NO_CONTEXT, contextAttr, 0);
-
+		EGLErrorStatus += (context != EGL14.EGL_NO_CONTEXT) ? "" : "eglCreateContext = " + EGL14.eglGetError() + " ";
+		
 		// make context to gather information
-		EGL14.eglMakeCurrent(display, surface, surface, context);
+		boolean bMakeCurrent = EGL14.eglMakeCurrent(display, surface, surface, context);
+		EGLErrorStatus += bMakeCurrent ? "" : "eglMakeCurrent = " + EGL14.eglGetError() + " ";
 
 		// grab info now
 		String Extensions = GLES20.glGetString(GLES20.GL_EXTENSIONS);
@@ -1155,6 +1171,10 @@ public class GameActivity extends $${gameActivitySuperClass}$$ implements Surfac
 		OpenGLDriver = GLES20.glGetString(GLES20.GL_VERSION);
 		boolean bES30Support = OpenGLDriver.contains("OpenGL ES 3.");
 
+		// write any gl error to status
+		int finalGLESError = GLES20.glGetError();
+		EGLErrorStatus += finalGLESError == GLES20.GL_NO_ERROR ? "getGraphicsInformation completed. " : "gl error = " + finalGLESError;
+
 		// tear it down
 		EGL14.eglMakeCurrent(display, EGL14.EGL_NO_SURFACE, EGL14.EGL_NO_SURFACE, EGL14.EGL_NO_CONTEXT);
 		EGL14.eglDestroySurface(display, surface);
@@ -1697,6 +1717,7 @@ public class GameActivity extends $${gameActivitySuperClass}$$ implements Surfac
 			}
 			catch (Exception e)
 			{
+				ConfigRulesStatus += "failed to create config rules reader. " + e;
 				reader = null;
 			}
 		}
@@ -2087,9 +2108,11 @@ public class GameActivity extends $${gameActivitySuperClass}$$ implements Surfac
 			catch (IOException ie)
 			{
 				Log.debug("failed to read configuration rules: " + ie);
+				ConfigRulesStatus += "failed to read configuration rules : " + ie+ " ";
 			}
 			if (reader != null)
 			{
+				ConfigRulesStatus += "configuration rules finished. " + ForceExitMessage;
 				try
 				{
 					reader.close();
