Minor performance improvement for logging

#ue4
#android

diff --git a/Engine/Build/Android/Java/src/com/epicgames/ue4/Logger.java b/Engine/Build/Android/Java/src/com/epicgames/ue4/Logger.java
--- a/Engine/Build/Android/Java/src/com/epicgames/ue4/Logger.java
+++ b/Engine/Build/Android/Java/src/com/epicgames/ue4/Logger.java
@@ -13,6 +13,7 @@ public class Logger
 	private static ILoggerCallback mCallback = null;
 	private final String mTag;
 	private final String mSecondaryTag;
+	private boolean bHasSecondaryTag;
 	private String mFormattedTag;
 	private String mFormattedMessageTag;
 	
@@ -53,6 +54,7 @@ public class Logger
 		} else {
 			mSecondaryTag = secondaryTag;
 		}
+		bHasSecondaryTag = !TextUtils.isEmpty(mSecondaryTag);
 	}
 
 	public void verbose(String Message)
@@ -120,7 +122,7 @@ public class Logger
 		if(mFormattedTag == null || bPrependChanged)
 		{
 			bPrependChanged = false;
-			if (!TextUtils.isEmpty(mSecondaryTag) && bPrependSecondaryTag)
+			if (bHasSecondaryTag && bPrependSecondaryTag)
 			{
 				mFormattedTag = mTag + "-" + mSecondaryTag;
 			}
@@ -134,7 +136,7 @@ public class Logger
 
 	private String getFormattedMessage(String message)
 	{
-		if (!TextUtils.isEmpty(mSecondaryTag) && !bPrependSecondaryTag)
+		if (bHasSecondaryTag && !bPrependSecondaryTag)
 		{
 			if(mFormattedMessageTag == null || bPrependChanged)
 			{
