Enable Emoji input for Android

diff --git a/Engine/Build/Android/Java/src/com/epicgames/ue4/GameActivity.java.template b/Engine/Build/Android/Java/src/com/epicgames/ue4/GameActivity.java.template
--- a/Engine/Build/Android/Java/src/com/epicgames/ue4/GameActivity.java.template
+++ b/Engine/Build/Android/Java/src/com/epicgames/ue4/GameActivity.java.template
@@ -5437,7 +5437,7 @@ public class GameActivity extends $${gameActivitySuperClass}$$ implements Surfac
 			{
 				emojiExcludeFilter = new EmojiExcludeFilter();
 			}
-			setFilters(new InputFilter[]{emojiExcludeFilter});
+			//setFilters(new InputFilter[]{emojiExcludeFilter});
 		}
 
 		@Override
