Handle redirects for streaming media player on Android

#ue4
#android

diff --git a/Engine/Build/Android/Java/src/com/epicgames/ue4/MediaPlayer14.java b/Engine/Build/Android/Java/src/com/epicgames/ue4/MediaPlayer14.java
--- a/Engine/Build/Android/Java/src/com/epicgames/ue4/MediaPlayer14.java
+++ b/Engine/Build/Android/Java/src/com/epicgames/ue4/MediaPlayer14.java
@@ -228,18 +228,47 @@ public class MediaPlayer14
 		return null;
 	}
 
-	 public static boolean RemoteFileExists(String URLName){
-        try {
-            HttpURLConnection.setFollowRedirects(false);
-            HttpURLConnection con =  (HttpURLConnection) new URL(URLName).openConnection();
-            con.setRequestMethod("HEAD");
-            return (con.getResponseCode() == HttpURLConnection.HTTP_OK);
-        }
-        catch (Exception e) {
-            e.printStackTrace();
-            return false;
-        }
-    }
+	public static String RemoteFileExists(String URLName)
+	{
+		// do not try more than 5 redirects, return final URL or null
+		int MaxRedirects = 5;
+		
+		// we set redirect to false and do it manually so we can handle redirect between HTTP/HTTPS which Java doesn't do
+		boolean restoreRedirects = HttpURLConnection.getFollowRedirects();
+		HttpURLConnection.setFollowRedirects(false);
+
+		while (MaxRedirects-- > 0)
+		{
+            try {
+				HttpURLConnection con = (HttpURLConnection) new URL(URLName).openConnection();
+				con.setRequestMethod("HEAD");
+				int responseCode = con.getResponseCode();
+				if (responseCode == HttpURLConnection.HTTP_OK)
+				{
+					con.disconnect();
+					HttpURLConnection.setFollowRedirects(restoreRedirects);
+					return URLName;
+				}
+				if (responseCode == HttpURLConnection.HTTP_SEE_OTHER || responseCode == HttpURLConnection.HTTP_MOVED_PERM || responseCode == HttpURLConnection.HTTP_MOVED_TEMP)
+				{
+					String Location = con.getHeaderField("Location");
+					URL url = con.getURL();
+					con.disconnect();
+					URLName = Location.contains("://") ? Location : url.getProtocol() + "://" + url.getHost() + Location;
+					continue;
+				}
+				HttpURLConnection.setFollowRedirects(restoreRedirects);
+				return null;
+			}
+			catch (Exception e) {
+				e.printStackTrace();
+				HttpURLConnection.setFollowRedirects(restoreRedirects);
+				return null;
+			}
+		}
+		HttpURLConnection.setFollowRedirects(restoreRedirects);
+		return null;
+	}
 
 	public boolean setDataSourceURL(
 		String UrlPath)
@@ -257,7 +286,8 @@ public class MediaPlayer14
 		audioTracks.clear();
 		videoTracks.clear();
 
-		if(!RemoteFileExists(UrlPath))
+		UrlPath = RemoteFileExists(UrlPath);
+		if (UrlPath == null)
 		{
 			return false;
 		}
