Fix configrules.txt operator bug (contributed by [REDACTED_PERSON])

#PR #6249
#ue4
#android

diff --git a/Engine/Build/Android/Java/src/com/epicgames/ue4/GameActivity.java.template b/Engine/Build/Android/Java/src/com/epicgames/ue4/GameActivity.java.template
--- a/Engine/Build/Android/Java/src/com/epicgames/ue4/GameActivity.java.template
+++ b/Engine/Build/Android/Java/src/com/epicgames/ue4/GameActivity.java.template
@@ -1541,7 +1541,7 @@ public class GameActivity extends $${gameActivitySuperClass}$$ implements Surfac
 				}
 				else if (CompareType.equals("CMP_GreaterEqual"))
 				{
-					if ((bNumericOperands && (SourceFloat <= MatchFloat)) || (!bNumericOperands && (Source.compareTo(MatchString) < 0)))
+					if ((bNumericOperands && (SourceFloat < MatchFloat)) || (!bNumericOperands && (Source.compareTo(MatchString) < 0)))
 					{
 						bConditionTrue = false;
 						break;
