Fix for issue when configrules.txt less than 80 bytes (contributed by [REDACTED_PERSON])

#PR #6237
#ue4
#android

diff --git a/Engine/Build/Android/Prebuilt/ConfigRulesTool/src/ConfigRulesTool.java b/Engine/Build/Android/Prebuilt/ConfigRulesTool/src/ConfigRulesTool.java
--- a/Engine/Build/Android/Prebuilt/ConfigRulesTool/src/ConfigRulesTool.java
+++ b/Engine/Build/Android/Prebuilt/ConfigRulesTool/src/ConfigRulesTool.java
@@ -92,7 +92,7 @@ public class ConfigRulesTool
 			int version = -1;
 			try
 			{
-				String versionLine = new String(bytesToCompress, 0, 80, "UTF-8");
+				String versionLine = new String(bytesToCompress, 0, sizeUncompressed, "UTF-8");
 				if (versionLine.startsWith("// version:"))
 				{
 					int eolIndex = versionLine.indexOf("\r");
