Small update to configrules tool

#ue4
#android

diff --git a/Engine/Build/Android/Prebuilt/ConfigRulesTool/src/ConfigRulesTool.java b/Engine/Build/Android/Prebuilt/ConfigRulesTool/src/ConfigRulesTool.java
--- a/Engine/Build/Android/Prebuilt/ConfigRulesTool/src/ConfigRulesTool.java
+++ b/Engine/Build/Android/Prebuilt/ConfigRulesTool/src/ConfigRulesTool.java
@@ -92,7 +92,8 @@ public class ConfigRulesTool
 			int version = -1;
 			try
 			{
-				String versionLine = new String(bytesToCompress, 0, sizeUncompressed, "UTF-8");
+				int versionBytes = sizeUncompressed < 80 ? sizeUncompressed : 80;
+				String versionLine = new String(bytesToCompress, 0, versionBytes, "UTF-8");
 				if (versionLine.startsWith("// version:"))
 				{
 					int eolIndex = versionLine.indexOf("\r");
