Verify launch notification event before registering it

#android

diff --git a/Engine/Build/Android/Java/src/com/epicgames/ue4/GameActivity.java.template b/Engine/Build/Android/Java/src/com/epicgames/ue4/GameActivity.java.template
--- a/Engine/Build/Android/Java/src/com/epicgames/ue4/GameActivity.java.template
+++ b/Engine/Build/Android/Java/src/com/epicgames/ue4/GameActivity.java.template
@@ -4680,30 +4680,34 @@ public class GameActivity extends $${gameActivitySuperClass}$$ implements Surfac
 
 	private void LocalNotificationCheckAppOpen()
 	{
+		localNotificationAppLaunched = false;
+		localNotificationLaunchActivationEvent = "";
+		localNotificationLaunchFireDate = 0;
+
 		Intent launchIntent = getIntent();
 		if (launchIntent != null)
 		{	
 			Bundle extrasBundle = launchIntent.getExtras();
 
 			localNotificationAppLaunched = launchIntent.getBooleanExtra("localNotificationAppLaunched", false);
-
-			if(localNotificationAppLaunched)
+			if (localNotificationAppLaunched)
 			{
-				localNotificationLaunchActivationEvent = extrasBundle.get("localNotificationLaunchActivationEvent").toString();
-				int notificationID = extrasBundle.getInt("localNotificationID");
-
-				LocalNotificationRemoveID(this, notificationID);
+				localNotificationLaunchActivationEvent = extrasBundle.getString("localNotificationLaunchActivationEvent");
+				if (localNotificationLaunchActivationEvent != null)
+				{
+					int notificationID = extrasBundle.getInt("localNotificationID");
+					LocalNotificationRemoveID(this, notificationID);
 
-				// TODO
-				localNotificationLaunchFireDate = 0; 
+					// TODO
+					localNotificationLaunchFireDate = 0; 
+				}
+				else
+				{
+					localNotificationAppLaunched = false;
+					localNotificationLaunchActivationEvent = "";
+				}
 			}
 		}
-		else
-		{
-			localNotificationAppLaunched = false;
-			localNotificationLaunchActivationEvent = "";
-			localNotificationLaunchFireDate = 0;
-		}
 	}
 
 	public static int LocalNotificationGetID(Context context)
