Improve Network Connection Manager's ability to detect bad network conditions

# Summary
This matches code that we recently added to iOS. Along with attempting to determine network connectivity with system methods, this also checks the ability to resolve a HEAD request with an external host `http://google.com/` in-order to verify that the network connection actually has a connection. This check is asynchronous with a 2 second timeout and has an exponential back-off retry with a maximum of up to 13 seconds when the check fails.

# Test Plan
1. Built and ran app with UeDevelopmentDebug
2. app launched
3. signed in
4. verified that could start download
5. Turned off wifi
6. Verified that saw network error
7. Turned wifi back on
8. Verified that network error went away automatically

Repeated above steps but instead of turning off wifi, had Scout test it by unplugging the network cable to a wifi router.

diff --git a/Engine/Build/Android/Java/src/com/epicgames/ue4/network/NetworkChangedListener.java b/Engine/Build/Android/Java/src/com/epicgames/ue4/network/NetworkChangedListener.java
deleted file mode 100644
diff --git a/Engine/Build/Android/Java/src/com/epicgames/ue4/network/NetworkChangedManager.java b/Engine/Build/Android/Java/src/com/epicgames/ue4/network/NetworkChangedManager.java
--- a/Engine/Build/Android/Java/src/com/epicgames/ue4/network/NetworkChangedManager.java
+++ b/Engine/Build/Android/Java/src/com/epicgames/ue4/network/NetworkChangedManager.java
@@ -1,92 +1,358 @@
+// Copyright 1998-2019 Epic Games, Inc. All Rights Reserved.
+
 package com.epicgames.ue4.network;
 
+import android.annotation.TargetApi;
 import android.content.Context;
 import android.net.ConnectivityManager;
 import android.net.Network;
+import android.net.NetworkCapabilities;
+import android.net.NetworkInfo;
 import android.net.NetworkRequest;
+import android.os.Build;
+import android.os.Handler;
 import android.support.annotation.NonNull;
+import android.support.annotation.Nullable;
 
 import com.epicgames.ue4.Logger;
 
+import java.io.IOException;
 import java.lang.ref.WeakReference;
+import java.net.HttpURLConnection;
+import java.net.MalformedURLException;
+import java.net.URL;
 import java.util.HashSet;
+import java.util.Iterator;
 import java.util.Set;
+import java.util.concurrent.Callable;
+import java.util.concurrent.ExecutorService;
+import java.util.concurrent.Executors;
+import java.util.concurrent.Future;
+import java.util.concurrent.TimeUnit;
 
-public final class NetworkChangedManager {
+public final class NetworkChangedManager implements NetworkConnectivityClient {
 
-	private static final Logger Log = new Logger("UE4", "NetworkChangedManager");
-	private static final String SYSTEM = "NetworkManager";
-	private static final String COMPONENT = "ConnectivityManager";
-	private static final String ACTION_ERROR = "InstanceNotAvailable";
+	private enum ConnectivityState {
+		CONNECTION_AVAILABLE,
+		NO_CONNECTION
+	}
 
-	private static NetworkChangedManager instance;
+	private static final int MAX_RETRY_SEC = 13;
+	private static final String HOST_RESOLUTION_ADDRESS = "http://google.com/";
+	private static final long HOSTNAME_RESOLUTION_TIMEOUT_MS = 2000;
 
-	@NonNull
-	private Set<WeakReference<NetworkChangedListener>> networkChangedListeners = new HashSet<>();
+	private static final Logger Log = new Logger("UE4-NetworkChangedManager");
 
+	private static NetworkChangedManager INSTANCE;
 	@NonNull
-	public static synchronized NetworkChangedManager getInstance() {
-		if (instance == null) {
-			instance = new NetworkChangedManager();
+	public static synchronized NetworkConnectivityClient getInstance() {
+		if (INSTANCE == null) {
+			INSTANCE = new NetworkChangedManager();
 		}
-		return instance;
+		return INSTANCE;
 	}
 
+	/*
+	 * References
+	 */
+	private ConnectivityManager connectivityManager;
+
+	@NonNull
+	private Set<WeakReference<Listener>> networkChangedListeners = new HashSet<>();
+
+	/*
+	 * Data
+	 */
+	private Set<String> networks = new HashSet<>();
+	@Nullable
+	private ConnectivityState currentState = null;
+
+	private boolean networkCheckInProgress = false;
+	private boolean retryScheduled = false;
+	private int retryCount = 0;
+
+	private Handler internalScheduler = new Handler();
+
 	private NetworkChangedManager() {
 	}
 
+	@Override
 	public void initNetworkCallback(@NonNull Context context) {
-		if (android.os.Build.VERSION.SDK_INT < 21)
-		{
+		if (Build.VERSION.SDK_INT < Build.VERSION_CODES.LOLLIPOP) {
 			// unsupported before Lollipop
 			return;
 		}
 
-		ConnectivityManager connectivityManager = (ConnectivityManager) context.getSystemService(Context.CONNECTIVITY_SERVICE);
-		if (connectivityManager != null) {
-			NetworkRequest.Builder builder = new NetworkRequest.Builder();
-			connectivityManager.registerNetworkCallback(
-				builder.build(),
-				new ConnectivityManager.NetworkCallback() {
-					@Override
-					public void onAvailable(Network network) {
-						Log.verbose("Network Available");
-						for (WeakReference<NetworkChangedListener> listenerWeakRef : networkChangedListeners) {
-							NetworkChangedListener listener = listenerWeakRef.get();
-							if (listener != null) {
-								listener.onNetworkAvailable(network);
-							} else {
-								removeListener(listenerWeakRef);
-							}
-						}
-					}
-
-					@Override
-					public void onLost(Network network) {
-						Log.verbose("Network Lost");
-						for (WeakReference<NetworkChangedListener> listenerWeakRef : networkChangedListeners) {
-							NetworkChangedListener listener = listenerWeakRef.get();
-							if (listener != null) {
-								listener.onNetworkLost(network);
-							} else {
-								removeListener(listenerWeakRef);
+		connectivityManager = (ConnectivityManager) context.getApplicationContext().getSystemService(Context.CONNECTIVITY_SERVICE);
+		if (connectivityManager == null) {
+			Log.error("Unable to start connectivityManager");
+			return;
+		}
+		NetworkRequest.Builder builder = new NetworkRequest.Builder();
+		builder.addCapability(NetworkCapabilities.NET_CAPABILITY_INTERNET);
+		if (Build.VERSION.SDK_INT >= Build.VERSION_CODES.M) {
+			builder.addCapability(NetworkCapabilities.NET_CAPABILITY_VALIDATED);
+		}
+		connectivityManager.registerNetworkCallback(builder.build(), connectivityListener);
+	}
+
+	@TargetApi(Build.VERSION_CODES.LOLLIPOP)
+	private ConnectivityManager.NetworkCallback connectivityListener = new ConnectivityManager.NetworkCallback() {
+		@Override
+		public void onAvailable(Network network) {
+			networks.add(network.toString());
+			Log.verbose("Network Available: " + network.toString());
+			checkNetworkConnectivity();
+		}
+
+		@Override
+		public void onLost(Network network) {
+			networks.remove(network.toString());
+			Log.verbose("Network Lost callback: " + network.toString());
+
+			if (networks.isEmpty()) {
+				Log.verbose("All Networks Lost");
+			} else {
+				Log.verbose("Not All Networks Lost");
+			}
+			checkNetworkConnectivity();
+		}
+
+		@Override
+		public void onCapabilitiesChanged(Network network, NetworkCapabilities networkCapabilities) {
+			if (Build.VERSION.SDK_INT >= Build.VERSION_CODES.M
+					&& !networkCapabilities.hasCapability(NetworkCapabilities.NET_CAPABILITY_VALIDATED)) {
+				Log.verbose("Network Capabilities changed, doesn't have validated net_capability");
+				networks.remove(network.toString());
+			} else if (!networkCapabilities.hasCapability(NetworkCapabilities.NET_CAPABILITY_INTERNET)) {
+				Log.verbose("Network Capabilities changed, has Internet: " + false);
+				networks.remove(network.toString());
+			} else {
+				Log.verbose("Network Capabilities changed, has Internet: " + true);
+				networks.add(network.toString());
+			}
+			checkNetworkConnectivity();
+		}
+	};
+
+	private Runnable retryRunnable = new Runnable() {
+		@Override
+		public void run() {
+			Log.verbose("Attempting to check for network connectivity again.");
+			retryCount++;
+			retryScheduled = false;
+			checkNetworkConnectivity();
+		}
+	};
+
+	private void scheduleRetry() {
+		if (networkChangedListeners.size() == 0) {
+			Log.verbose("No listeners so not retrying. When a listener is added the connection will be checked.");
+			return;
+		}
+		if (retryScheduled || networkCheckInProgress) {
+			return;
+		}
+		retryScheduled = true;
+		internalScheduler.removeCallbacksAndMessages(retryRunnable);
+		internalScheduler.postDelayed(retryRunnable, calculateRetryDelay());
+	}
+
+	/**
+	 * @return a delay in MS that increases exponentially but with a max value of
+	 * {@link NetworkChangedManager#MAX_RETRY_SEC} seconds
+	 */
+	private long calculateRetryDelay() {
+		return (long) (Math.min(MAX_RETRY_SEC, Math.pow(2, retryCount)) * 1000);
+	}
+
+	private void clearRetry() {
+		internalScheduler.removeCallbacksAndMessages(retryRunnable);
+		retryCount = 0;
+		retryScheduled = false;
+	}
+
+	private void setNetworkState(ConnectivityState state) {
+		if (currentState == state) {
+			Log.verbose("Connectivity hasn't changed. Current state: " + currentState);
+			if (currentState != ConnectivityState.CONNECTION_AVAILABLE) {
+				scheduleRetry();
+			}
+			return;
+		}
+		currentState = state;
+		fireNetworkChangeListeners(state);
+		Log.verbose("Network connectivity changed. New connectivity state: " + state);
+
+		if (currentState != ConnectivityState.CONNECTION_AVAILABLE) {
+			scheduleRetry();
+		} else {
+			clearRetry();
+		}
+	}
+
+	private void checkNetworkConnectivity() {
+		ConnectivityState naiveNetworkState = calculateNetworkConnectivityNaively();
+		if (naiveNetworkState != ConnectivityState.CONNECTION_AVAILABLE) {
+			setNetworkState(ConnectivityState.NO_CONNECTION);
+			return;
+		} else if (currentState == null) {
+			Log.verbose("No network state set yet, setting naive network state checking connection fully.");
+			setNetworkState(naiveNetworkState);
+		}
+
+		if (networkCheckInProgress) {
+			return;
+		}
+
+		networkCheckInProgress = true;
+		internalScheduler.post(new Runnable() {
+			@Override
+			public void run() {
+				boolean successfullyConnectedToServer = false;
+				try {
+					ExecutorService executor = Executors.newSingleThreadExecutor();
+					Future<Boolean> future = executor.submit(new Callable<Boolean>() {
+						@Override
+						public Boolean call() {
+							HttpURLConnection urlConnection = null;
+							try {
+								URL url = new URL(HOST_RESOLUTION_ADDRESS);
+								urlConnection = (HttpURLConnection) url.openConnection();
+								urlConnection.setRequestMethod("HEAD");
+								urlConnection.getInputStream().close();
+							} catch (MalformedURLException e) {
+								Log.error("Malformed URL, this should never happen. Please fix, url: " + HOST_RESOLUTION_ADDRESS);
+								e.printStackTrace();
+								return false;
+							} catch (IOException e) {
+								Log.verbose("Unable to connect to: " + HOST_RESOLUTION_ADDRESS);
+								return false;
+							} finally {
+								if (urlConnection != null) {
+									urlConnection.disconnect();
+								}
 							}
+							return true;
 						}
-					}
+					});
+					successfullyConnectedToServer = future.get(HOSTNAME_RESOLUTION_TIMEOUT_MS, TimeUnit.MILLISECONDS);
+					future.cancel(true);
+					executor.shutdownNow();
+				} catch (Exception ignored) {
+				}
+				networkCheckInProgress = false;
+				if (successfullyConnectedToServer) {
+					setNetworkState(ConnectivityState.CONNECTION_AVAILABLE);
+				} else {
+					setNetworkState(ConnectivityState.NO_CONNECTION);
+					scheduleRetry();
 				}
-			);
+				Log.verbose("Full network check complete. State: " + currentState);
+			}
+		});
+	}
+
+	@NonNull
+	private ConnectivityState calculateNetworkConnectivityNaively() {
+		if (networks.isEmpty() || connectivityManager == null) {
+			return ConnectivityState.NO_CONNECTION;
 		} else {
-			Log.error("Unable to start connectivityManager");
+			NetworkInfo activeNetwork = connectivityManager.getActiveNetworkInfo();
+
+			boolean isConnected = (activeNetwork != null && activeNetwork.isAvailable() && activeNetwork.isConnected());
+			if (Build.VERSION.SDK_INT >= Build.VERSION_CODES.M) {
+				if (activeNetwork != null
+						&& !connectivityManager.getNetworkCapabilities(connectivityManager.getActiveNetwork()).hasCapability(NetworkCapabilities.NET_CAPABILITY_VALIDATED)) {
+					isConnected = false;
+				}
+			}
+			if (isConnected) {
+				return ConnectivityState.CONNECTION_AVAILABLE;
+			} else {
+				return ConnectivityState.NO_CONNECTION;
+			}
 		}
 	}
 
-	@SuppressWarnings({"unused", "UnusedReturnValue"})
-	public boolean addListener(WeakReference<NetworkChangedListener> listener) {
-		return networkChangedListeners.add(listener);
+	@Override
+	public boolean addListener(Listener listener) {
+		return addListener(listener, false);
 	}
 
-	@SuppressWarnings({"unused", "UnusedReturnValue"})
-	public boolean removeListener(WeakReference<NetworkChangedListener> listener) {
-		return networkChangedListeners.remove(listener);
+	@Override
+	public boolean addListener(Listener listener, boolean fireImmediately) {
+		if (Build.VERSION.SDK_INT < Build.VERSION_CODES.LOLLIPOP) {
+			// unsupported before Lollipop
+			return false;
+		}
+
+		for (WeakReference<Listener> currentListener : networkChangedListeners) {
+			if (currentListener.get() == listener) {
+				return false;
+			}
+		}
+		networkChangedListeners.add(new WeakReference<>(listener));
+		if (networkChangedListeners.size() == 1) {
+			/*
+			 * Start checking connectivity when we go from 0 -> 1 listeners.
+			 * When no listeners are registered checking connectivity is automatically stopped.
+			 */
+			checkNetworkConnectivity();
+		}
+		// Current state will never be null after calling checkNetworkConnectivity
+		if (fireImmediately && currentState != null) {
+			fireNetworkChangeListenerInternal(listener, currentState);
+		}
+		return true;
+	}
+
+	@Override
+	public boolean removeListener(Listener listener) {
+		if (Build.VERSION.SDK_INT < Build.VERSION_CODES.LOLLIPOP) {
+			// unsupported before Lollipop
+			return false;
+		}
+
+		Iterator<WeakReference<Listener>> changeListenersIterator = networkChangedListeners.iterator();
+		while (changeListenersIterator.hasNext()) {
+			WeakReference<Listener> currentListener = changeListenersIterator.next();
+			if (currentListener.get() == listener) {
+				changeListenersIterator.remove();
+				return true;
+			}
+		}
+		/*
+		 * If there are no listeners, clear any retries. When a listener is added we'll check the
+		 * connection again.
+		 */
+		if (networkChangedListeners.size() == 0) {
+			clearRetry();
+		}
+		return false;
+	}
+
+	private void fireNetworkChangeListeners(ConnectivityState state) {
+		Iterator<WeakReference<Listener>> changeListenersIterator = networkChangedListeners.iterator();
+		while (changeListenersIterator.hasNext()) {
+			WeakReference<Listener> listenerWeakReference = changeListenersIterator.next();
+			Listener listener = listenerWeakReference.get();
+			if (listener == null) {
+				changeListenersIterator.remove();
+			} else {
+				fireNetworkChangeListenerInternal(listener, state);
+			}
+		}
+	}
+
+	private void fireNetworkChangeListenerInternal(Listener listener, ConnectivityState state) {
+		switch (state) {
+			case NO_CONNECTION:
+				listener.onNetworkLost();
+				break;
+			case CONNECTION_AVAILABLE:
+				listener.onNetworkAvailable();
+				break;
+		}
 	}
 }
diff --git a/Engine/Build/Android/Java/src/com/epicgames/ue4/network/NetworkConnectivityClient.java b/Engine/Build/Android/Java/src/com/epicgames/ue4/network/NetworkConnectivityClient.java
new file mode 100644
--- /dev/null
+++ b/Engine/Build/Android/Java/src/com/epicgames/ue4/network/NetworkConnectivityClient.java
@@ -0,0 +1,36 @@
+// Copyright 1998-2019 Epic Games, Inc. All Rights Reserved.
+
+package com.epicgames.ue4.network;
+
+import android.content.Context;
+import android.support.annotation.NonNull;
+
+public interface NetworkConnectivityClient {
+	public interface Listener {
+		void onNetworkAvailable();
+
+		void onNetworkLost();
+	}
+
+	void initNetworkCallback(@NonNull Context context);
+
+	/**
+	 * See {@link NetworkConnectivityClient#addListener(Listener, boolean)}
+	 */
+	boolean addListener(Listener listener);
+
+	/**
+	 * @param listener The listener to add. Will be stored as a weak reference so a hard reference
+	 *                 must be saved externally.
+	 * @param fireImmediately Whether to trigger the listener with the current network state
+	 *                        immediately after adding.
+	 * @return Whether the change listener was added. Will be false if already registered.
+	 */
+	boolean addListener(Listener listener, boolean fireImmediately);
+
+	/**
+	 * Remove a given listener.
+	 * @return Whether the change listener was removed. Will be false if not currently registered.
+	 */
+	boolean removeListener(Listener listener);
+}
