Fix blocking code in NetworkChangedManager

# Summary
Since the HTTP url open connection is a blocking call. It needs to be run on a seperate thread. The call .get on a future is blocking though, meaning we were still running a blocking call on the main looper. This updates that code to instead add a timeout runnable to the local handler, to cancel the executor if it hasn't completed in the allotted time.

Added some tweaks to match comments made by [REDACTED_HANDLE]
 - Refactor INSTANCE to be lower case
 - Tweak construction of Logger to match current pattern
 - Copied over correct NetworkChangedManager to remove System static call that was accidently included in the NativeWrapper

# Test Plan
1. Built and ran app with UeDevelopmentDebug
2. app launched
3. signed in
4. verified that could start download
5. Turned off wifi
6. Verified that saw network error
7. Turned wifi back on
8. Verified that network error went away automatically

diff --git a/Engine/Build/Android/Java/src/com/epicgames/ue4/network/NetworkChangedManager.java b/Engine/Build/Android/Java/src/com/epicgames/ue4/network/NetworkChangedManager.java
--- a/Engine/Build/Android/Java/src/com/epicgames/ue4/network/NetworkChangedManager.java
+++ b/Engine/Build/Android/Java/src/com/epicgames/ue4/network/NetworkChangedManager.java
@@ -38,18 +38,18 @@ public final class NetworkChangedManager implements NetworkConnectivityClient {
 	}
 
 	private static final int MAX_RETRY_SEC = 13;
-	private static final String HOST_RESOLUTION_ADDRESS = "http://google.com/";
+	private static final String HOST_RESOLUTION_ADDRESS = "http://example.com/";
 	private static final long HOSTNAME_RESOLUTION_TIMEOUT_MS = 2000;
 
-	private static final Logger Log = new Logger("UE4-NetworkChangedManager");
+	private static final Logger Log = new Logger("UE4", "NetworkChangedManager");
 
-	private static NetworkChangedManager INSTANCE;
+	private static NetworkChangedManager instance;
 	@NonNull
 	public static synchronized NetworkConnectivityClient getInstance() {
-		if (INSTANCE == null) {
-			INSTANCE = new NetworkChangedManager();
+		if (instance == null) {
+			instance = new NetworkChangedManager();
 		}
-		return INSTANCE;
+		return instance;
 	}
 
 	/*
@@ -206,49 +206,52 @@ public final class NetworkChangedManager implements NetworkConnectivityClient {
 		}
 
 		networkCheckInProgress = true;
-		internalScheduler.post(new Runnable() {
+		ExecutorService executor = Executors.newSingleThreadExecutor();
+		Runnable timeoutRunnable = new Runnable() {
 			@Override
 			public void run() {
-				boolean successfullyConnectedToServer = false;
+				Log.verbose("Unable to connect to: " + HOST_RESOLUTION_ADDRESS);
+				networkCheckInProgress = false;
+				executor.shutdownNow();
+				setNetworkState(ConnectivityState.NO_CONNECTION);
+				scheduleRetry();
+			}
+		};
+		internalScheduler.postDelayed(timeoutRunnable, HOSTNAME_RESOLUTION_TIMEOUT_MS);
+		executor.execute(new Runnable() {
+			@Override
+			public void run() {
+				HttpURLConnection urlConnection = null;
+				boolean connectedSuccessfully = false;
 				try {
-					ExecutorService executor = Executors.newSingleThreadExecutor();
-					Future<Boolean> future = executor.submit(new Callable<Boolean>() {
-						@Override
-						public Boolean call() {
-							HttpURLConnection urlConnection = null;
-							try {
-								URL url = new URL(HOST_RESOLUTION_ADDRESS);
-								urlConnection = (HttpURLConnection) url.openConnection();
-								urlConnection.setRequestMethod("HEAD");
-								urlConnection.getInputStream().close();
-							} catch (MalformedURLException e) {
-								Log.error("Malformed URL, this should never happen. Please fix, url: " + HOST_RESOLUTION_ADDRESS);
-								e.printStackTrace();
-								return false;
-							} catch (IOException e) {
-								Log.verbose("Unable to connect to: " + HOST_RESOLUTION_ADDRESS);
-								return false;
-							} finally {
-								if (urlConnection != null) {
-									urlConnection.disconnect();
-								}
-							}
-							return true;
-						}
-					});
-					successfullyConnectedToServer = future.get(HOSTNAME_RESOLUTION_TIMEOUT_MS, TimeUnit.MILLISECONDS);
-					future.cancel(true);
-					executor.shutdownNow();
-				} catch (Exception ignored) {
+					URL url = new URL(HOST_RESOLUTION_ADDRESS);
+					urlConnection = (HttpURLConnection) url.openConnection();
+					urlConnection.setRequestMethod("HEAD");
+					urlConnection.setConnectTimeout((int) HOSTNAME_RESOLUTION_TIMEOUT_MS);
+					urlConnection.setReadTimeout((int) HOSTNAME_RESOLUTION_TIMEOUT_MS);
+					urlConnection.getInputStream().close();
+					connectedSuccessfully = true;
+				} catch (MalformedURLException e) {
+					Log.error("Malformed URL, this should never happen. Please fix, url: " + HOST_RESOLUTION_ADDRESS);
+				} catch (IOException e) {
+					Log.verbose("Unable to connect to: " + HOST_RESOLUTION_ADDRESS);
+				} finally {
+					if (urlConnection != null) {
+						urlConnection.disconnect();
+					}
 				}
+
+				internalScheduler.removeCallbacks(timeoutRunnable);
 				networkCheckInProgress = false;
-				if (successfullyConnectedToServer) {
+				if (connectedSuccessfully) {
 					setNetworkState(ConnectivityState.CONNECTION_AVAILABLE);
 				} else {
 					setNetworkState(ConnectivityState.NO_CONNECTION);
 					scheduleRetry();
 				}
 				Log.verbose("Full network check complete. State: " + currentState);
+
+				executor.shutdownNow();
 			}
 		});
 	}
