Fix not declaring variable accessed by an inner class as final

# Summary
See title, variable needed to be declared as final.

diff --git a/Engine/Build/Android/Java/src/com/epicgames/ue4/network/NetworkChangedManager.java b/Engine/Build/Android/Java/src/com/epicgames/ue4/network/NetworkChangedManager.java
--- a/Engine/Build/Android/Java/src/com/epicgames/ue4/network/NetworkChangedManager.java
+++ b/Engine/Build/Android/Java/src/com/epicgames/ue4/network/NetworkChangedManager.java
@@ -24,11 +24,8 @@ import java.net.URL;
 import java.util.HashSet;
 import java.util.Iterator;
 import java.util.Set;
-import java.util.concurrent.Callable;
 import java.util.concurrent.ExecutorService;
 import java.util.concurrent.Executors;
-import java.util.concurrent.Future;
-import java.util.concurrent.TimeUnit;
 
 public final class NetworkChangedManager implements NetworkConnectivityClient {
 
@@ -206,8 +203,8 @@ public final class NetworkChangedManager implements NetworkConnectivityClient {
 		}
 
 		networkCheckInProgress = true;
-		ExecutorService executor = Executors.newSingleThreadExecutor();
-		Runnable timeoutRunnable = new Runnable() {
+		final ExecutorService executor = Executors.newSingleThreadExecutor();
+		final Runnable timeoutRunnable = new Runnable() {
 			@Override
 			public void run() {
 				Log.verbose("Unable to connect to: " + HOST_RESOLUTION_ADDRESS);
