Fix NPE in network capability validation code

#Summary
Remove useless network capability validated code

The crash was caused by failing to check for null in responses from calls to some system methods.
In general this code doesn't help much since this specific check was added prior to adding the more thorough HEAD request to example.com.
By removing it, it should fix the null issue, and in-fact make the naive internet check more in-line with what's expected which is a very simple network check.
This naive check always followed by attempting to connect to example.com which provides a far more guaranteed check into whether the internet is in-fact connected.

#Test Plan
Connected bluetooth headset
Launched app
launched game
Disconnected bluetooth
Verified no crash

diff --git a/Engine/Build/Android/Java/src/com/epicgames/ue4/network/NetworkChangedManager.java b/Engine/Build/Android/Java/src/com/epicgames/ue4/network/NetworkChangedManager.java
--- a/Engine/Build/Android/Java/src/com/epicgames/ue4/network/NetworkChangedManager.java
+++ b/Engine/Build/Android/Java/src/com/epicgames/ue4/network/NetworkChangedManager.java
@@ -258,16 +258,9 @@ public final class NetworkChangedManager implements NetworkConnectivityClient {
 		if (networks.isEmpty() || connectivityManager == null) {
 			return ConnectivityState.NO_CONNECTION;
 		} else {
-			NetworkInfo activeNetwork = connectivityManager.getActiveNetworkInfo();
+			NetworkInfo activeNetworkInfo = connectivityManager.getActiveNetworkInfo();
 
-			boolean isConnected = (activeNetwork != null && activeNetwork.isAvailable() && activeNetwork.isConnected());
-			if (Build.VERSION.SDK_INT >= Build.VERSION_CODES.M) {
-				if (activeNetwork != null
-						&& !connectivityManager.getNetworkCapabilities(connectivityManager.getActiveNetwork()).hasCapability(NetworkCapabilities.NET_CAPABILITY_VALIDATED)) {
-					isConnected = false;
-				}
-			}
-			if (isConnected) {
+			if (activeNetworkInfo != null && activeNetworkInfo.isAvailable() && activeNetworkInfo.isConnected()) {
 				return ConnectivityState.CONNECTION_AVAILABLE;
 			} else {
 				return ConnectivityState.NO_CONNECTION;
