Fixed an issue with the Android virtual keyboard mangling emoji when using backspace.

#tests Compile Development Client Android

diff --git a/Engine/Build/Android/Java/src/com/epicgames/ue4/GameActivity.java.template b/Engine/Build/Android/Java/src/com/epicgames/ue4/GameActivity.java.template
--- a/Engine/Build/Android/Java/src/com/epicgames/ue4/GameActivity.java.template
+++ b/Engine/Build/Android/Java/src/com/epicgames/ue4/GameActivity.java.template
@@ -39,6 +39,7 @@ import java.util.HashMap;
 import java.util.ArrayList;
 import java.text.DecimalFormat;
 import java.text.DecimalFormatSymbols;
+import java.text.BreakIterator;
 import java.util.zip.Inflater;
 import java.util.Stack;
 import java.util.Timer;
@@ -5565,13 +5566,23 @@ public class GameActivity extends $${gameActivitySuperClass}$$ implements Surfac
 				{ 
 					//insert
 					text.insert(selStart, newString);
-				} 
+				}
 				else if(selStart > 0) 
-				{ 
-					//delete last character
-					selStart--;
-					text.replace(selStart, selStart + 1, "");
-				} 
+				{
+					// selection start and end are the same, but in the middle of the string
+					// this means we want a single "delete character to left of selection" (backspace)
+
+					// using BreakIterator to ensure we capture proper emoji / Unicode grapheme boundaries
+					BreakIterator finder = BreakIterator.getCharacterInstance();
+					finder.setText(owner.getText().toString());
+
+					selStart = finder.preceding(selStart);
+					if (selStart == BreakIterator.DONE)
+					{
+						selStart = 0;
+					}
+					text.replace(selStart, selEnd, "");
+				}
 
 				if(newString.length() == 0)
 				{
