Add another check for 64-bit ABI support

#android

diff --git a/Engine/Build/Android/Java/src/com/epicgames/ue4/GameActivity.java.template b/Engine/Build/Android/Java/src/com/epicgames/ue4/GameActivity.java.template
--- a/Engine/Build/Android/Java/src/com/epicgames/ue4/GameActivity.java.template
+++ b/Engine/Build/Android/Java/src/com/epicgames/ue4/GameActivity.java.template
@@ -1593,6 +1593,12 @@ public class GameActivity extends $${gameActivitySuperClass}$$ implements Surfac
 			bHasNEON = true;
 		}
 
+        String abilist = getProp("ro.product.cpu.abilist");
+        if (abilist.contains("arm64-v8a"))
+        {
+            bIsARM64 = true;
+        }
+
 		// check each processor featureset
 		for (int id=0; id < processorCount; id++)
 		{
