Fix issue with attempting to aquire WRITE_EXTERNAL_STORAGE when we already have it on non-shipping builds

#ue4
#android

diff --git a/Engine/Build/Android/Java/src/com/epicgames/ue4/GameActivity.java.template b/Engine/Build/Android/Java/src/com/epicgames/ue4/GameActivity.java.template
--- a/Engine/Build/Android/Java/src/com/epicgames/ue4/GameActivity.java.template
+++ b/Engine/Build/Android/Java/src/com/epicgames/ue4/GameActivity.java.template
@@ -2294,23 +2294,30 @@ public class GameActivity extends $${gameActivitySuperClass}$$ implements Surfac
 			Log.debug(e.getMessage());
 		}
 
+		boolean bIsShipping = nativeIsShippingBuild();
 		if (ANDROID_BUILD_VERSION >= 23 && targetSdkVersion >= 23) //23 is the API level (Marshmallow) where runtime permission handling is available
 		{
-			Log.debug("Target SDK is " + targetSdkVersion + ".  This may cause issues if permissions are denied by the user." );				
+			Log.debug("Target SDK is " + targetSdkVersion + ".  This may cause issues if permissions are denied by the user." );	
+			if (!bIsShipping)
+			{
+				// For non-shipping builds we may need this permission because we write to the sdcard for logs
+				if (!PermissionHelper.checkPermission("android.permission.WRITE_EXTERNAL_STORAGE"))
+				{
+					Log.debug("android.permission.WRITE_EXTERNAL_STORAGE was not granted; this may cause issues with file access." );
+					// we do not do this currently since we must wait until it is granted before accessing files
+					// SplashActivity takes care of properly gating this in the cases it is needed
+					//PermissionHelper.acquirePermissions(new String[] {"android.permission.WRITE_EXTERNAL_STORAGE"}, this);
+				}
+			}
 		}
 
 		Logger.RegisterCallback(this);
 		
 		// Suppress java logs in Shipping builds
-		if (nativeIsShippingBuild())
+		if (bIsShipping)
 		{
 			Logger.SuppressLogs();
 		}
-		else
-		{
-			//For non-shipping builds we need to request this permission bc we write to the sdcard for logs
-			PermissionHelper.acquirePermissions(new String[] {"android.permission.WRITE_EXTERNAL_STORAGE"}, this);
-		}
 
 		// update memory stats every 10 seconds
 		memoryRunnable = new Runnable() {
