Exposed a `checkConnectivity` method in NetworkConnectivityClient to let the Android native app use.

# Summary
It won't affect other parts of the app.
Then users can manually check the connectivity by clicking RETRY button from the Connection Error UI.

To trigger the connectivity checking immediately and to make sure this is only one-time connectivity checking, we call `checkNetworkConnectivity` directly to avoid to affect the scheduled auto-retry and avoid to change the current retry count. And if a check is already in progress (either triggered manually or scheduled), a new check won't be triggered.

#android

diff --git a/Engine/Build/Android/Java/src/com/epicgames/ue4/network/NetworkChangedManager.java b/Engine/Build/Android/Java/src/com/epicgames/ue4/network/NetworkChangedManager.java
--- a/Engine/Build/Android/Java/src/com/epicgames/ue4/network/NetworkChangedManager.java
+++ b/Engine/Build/Android/Java/src/com/epicgames/ue4/network/NetworkChangedManager.java
@@ -350,4 +350,9 @@ public final class NetworkChangedManager implements NetworkConnectivityClient {
 				break;
 		}
 	}
+	
+	@Override
+	public void checkConnectivity() {
+		checkNetworkConnectivity()
+	}
 }
diff --git a/Engine/Build/Android/Java/src/com/epicgames/ue4/network/NetworkConnectivityClient.java b/Engine/Build/Android/Java/src/com/epicgames/ue4/network/NetworkConnectivityClient.java
--- a/Engine/Build/Android/Java/src/com/epicgames/ue4/network/NetworkConnectivityClient.java
+++ b/Engine/Build/Android/Java/src/com/epicgames/ue4/network/NetworkConnectivityClient.java
@@ -33,4 +33,9 @@ public interface NetworkConnectivityClient {
 	 * @return Whether the change listener was removed. Will be false if not currently registered.
 	 */
 	boolean removeListener(Listener listener);
+	
+	/**
+	 * Check for network connectivity
+	 */
+	void checkConnectivity();
 }
