Fix the compiling error: "error: ';' expected"

#android

diff --git a/Engine/Build/Android/Java/src/com/epicgames/ue4/network/NetworkChangedManager.java b/Engine/Build/Android/Java/src/com/epicgames/ue4/network/NetworkChangedManager.java
--- a/Engine/Build/Android/Java/src/com/epicgames/ue4/network/NetworkChangedManager.java
+++ b/Engine/Build/Android/Java/src/com/epicgames/ue4/network/NetworkChangedManager.java
@@ -353,6 +353,6 @@ public final class NetworkChangedManager implements NetworkConnectivityClient {
 	
 	@Override
 	public void checkConnectivity() {
-		checkNetworkConnectivity()
+		checkNetworkConnectivity();
 	}
 }
