Add bDeviceSupported to GameActivity for configrules to modfy

#ue4
#android

diff --git a/Engine/Build/Android/Java/src/com/epicgames/ue4/GameActivity.java.template b/Engine/Build/Android/Java/src/com/epicgames/ue4/GameActivity.java.template
--- a/Engine/Build/Android/Java/src/com/epicgames/ue4/GameActivity.java.template
+++ b/Engine/Build/Android/Java/src/com/epicgames/ue4/GameActivity.java.template
@@ -475,6 +475,9 @@ public class GameActivity extends $${gameActivitySuperClass}$$ implements Surfac
 	public String TextureFormats = "ETC1";
 	public boolean bSupportsFloatingPointRenderTargets = false;
 	private String EGLErrorStatus = "UNSET";
+	
+	/** Flag for not officially supported devices */
+	public boolean bSupportedDevice = true;
 
 	/** Used for LocalNotification support*/
 	private boolean localNotificationAppLaunched = false;
@@ -1644,6 +1647,7 @@ public class GameActivity extends $${gameActivitySuperClass}$$ implements Surfac
 		variables.put("memory", Integer.toString(totalMemory));
 
 		//variables.put("Profile", "Android");
+		variables.put("unsupportedDevice", "false");
 		variables.put("OverrideProfile", "true");
 		variables.put("hardware", hardware);
 		variables.put("ro.hardware", roHardware);
@@ -2163,6 +2167,10 @@ public class GameActivity extends $${gameActivitySuperClass}$$ implements Surfac
 		}
 
 		// notify AndroidProcess of the core masks and affinity enable state
+		if (variables.containsKey("unsupportedDevice"))
+		{
+			bSupportedDevice = variables.get("unsupportedDevice").equals("false");
+		}
 		boolean bUseAffinity = variables.get("useAffinity").equals("true");
 		int littleCoreMask = Integer.decode(variables.get("littleCoreMask"));
 		int bigCoreMask = Integer.decode(variables.get("bigCoreMask"));
