Add new features from Google Play Billing API
Create Blueprint Nodes for StoreV2
Modify Match3 to use SToreV2 blueprint nodes

diff --git a/Engine/Build/Android/Java/src/com/android/vending/billing/util/IabException.java b/Engine/Build/Android/Java/src/com/android/vending/billing/util/IabException.java
deleted file mode 100644
diff --git a/Engine/Build/Android/Java/src/com/android/vending/billing/util/IabHelper.java b/Engine/Build/Android/Java/src/com/android/vending/billing/util/IabHelper.java
deleted file mode 100644
diff --git a/Engine/Build/Android/Java/src/com/android/vending/billing/util/IabResult.java b/Engine/Build/Android/Java/src/com/android/vending/billing/util/IabResult.java
deleted file mode 100644
diff --git a/Engine/Build/Android/Java/src/com/android/vending/billing/util/Inventory.java b/Engine/Build/Android/Java/src/com/android/vending/billing/util/Inventory.java
deleted file mode 100644
diff --git a/Engine/Build/Android/Java/src/com/android/vending/billing/util/Purchase.java b/Engine/Build/Android/Java/src/com/android/vending/billing/util/Purchase.java
deleted file mode 100644
diff --git a/Engine/Build/Android/Java/src/com/android/vending/billing/util/Security.java b/Engine/Build/Android/Java/src/com/android/vending/billing/util/Security.java
deleted file mode 100644
diff --git a/Engine/Build/Android/Java/src/com/android/vending/billing/util/SkuDetails.java b/Engine/Build/Android/Java/src/com/android/vending/billing/util/SkuDetails.java
deleted file mode 100644
diff --git a/Engine/Build/Android/Java/src/com/epicgames/ue4/GameActivity.java.template b/Engine/Build/Android/Java/src/com/epicgames/ue4/GameActivity.java.template
--- a/Engine/Build/Android/Java/src/com/epicgames/ue4/GameActivity.java.template
+++ b/Engine/Build/Android/Java/src/com/epicgames/ue4/GameActivity.java.template
@@ -4524,14 +4524,7 @@ public class GameActivity extends $${gameActivitySuperClass}$$ implements Surfac
 		}
 		else if( IapStoreHelper != null )
 		{
-			if(!IapStoreHelper.onActivityResult(requestCode, resultCode, data))
-			{
 				super.onActivityResult(requestCode, resultCode, data);
-			}
-			else
-			{
-				Log.debug("[JAVA] - Store Helper handled onActivityResult");
-			}
 		}
 		else
 		{
diff --git a/Engine/Build/Android/Java/src/com/epicgames/ue4/GooglePlayStoreHelper.java b/Engine/Build/Android/Java/src/com/epicgames/ue4/GooglePlayStoreHelper.java
--- a/Engine/Build/Android/Java/src/com/epicgames/ue4/GooglePlayStoreHelper.java
+++ b/Engine/Build/Android/Java/src/com/epicgames/ue4/GooglePlayStoreHelper.java
@@ -15,53 +15,17 @@ import android.os.Handler;
 import android.os.IBinder;
 import android.os.Looper;
 import android.os.RemoteException;
-import com.android.vending.billing.util.Base64;
-import com.android.vending.billing.IInAppBillingService;
-import com.android.vending.billing.util.Purchase;
+import com.android.vending.billing.util.Base64; //WMM should use different base64 here.
+import com.android.billingclient.api.*;
 import java.util.ArrayList;
 import java.util.List;
 import org.json.JSONException;
-import org.json.JSONObject;
+import org.json.JSONObject; 
 
 public class GooglePlayStoreHelper implements StoreHelper
 {
-    // Billing response codes
-	// taken from http://developer.android.com/google/play/billing/billing_reference.html Table 1
-    public static final int BILLING_RESPONSE_RESULT_OK = 0;
-    public static final int BILLING_RESPONSE_RESULT_USER_CANCELED = 1;
-	public static final int BILLING_RESPONSE_RESULT_SERVICE_UNAVAILABLE = 2;
-    public static final int BILLING_RESPONSE_RESULT_BILLING_UNAVAILABLE = 3;
-    public static final int BILLING_RESPONSE_RESULT_ITEM_UNAVAILABLE = 4;
-    public static final int BILLING_RESPONSE_RESULT_DEVELOPER_ERROR = 5;
-    public static final int BILLING_RESPONSE_RESULT_ERROR = 6;
-    public static final int BILLING_RESPONSE_RESULT_ITEM_ALREADY_OWNED = 7;
-    public static final int BILLING_RESPONSE_RESULT_ITEM_NOT_OWNED = 8;
-
-	// Keys for the responses from InAppBillingService
-	// taken from http://developer.android.com/google/play/billing/billing_reference.html Table 3
-	public static final String RESPONSE_CODE = "RESPONSE_CODE";
-	public static final String RESPONSE_INAPP_PURCHASE_DATA = "INAPP_PURCHASE_DATA";
-	public static final String RESPONSE_INAPP_SIGNATURE = "INAPP_DATA_SIGNATURE";
-
-	// taken from https://developer.android.com/google/play/billing/billing_integrate.html#purchase
-	public static final String RESPONSE_GET_SKU_DETAILS_LIST = "DETAILS_LIST";
-	public static final String RESPONSE_BUY_INTENT = "BUY_INTENT";
-	public static final String RESPONSE_INAPP_ITEM_LIST = "INAPP_PURCHASE_ITEM_LIST";
-	public static final String RESPONSE_INAPP_PURCHASE_DATA_LIST = "INAPP_PURCHASE_DATA_LIST";
-	public static final String RESPONSE_INAPP_SIGNATURE_LIST = "INAPP_DATA_SIGNATURE_LIST";
-	public static final String INAPP_CONTINUATION_TOKEN = "INAPP_CONTINUATION_TOKEN";
-
-    // Max number of offers to query at once (see getSkuDetails in IInAppbillingServer.aidl)
-	public static final int QueryPurchaseMaxCount = 20;
-
-	// Item types
-	public static final String ITEM_TYPE_INAPP = "inapp";
-
-	// some fields on the getSkuDetails response bundle
-	public static final String GET_SKU_DETAILS_ITEM_LIST = "ITEM_ID_LIST";
-
 	// Our IAB helper interface provided by google.
-	private IInAppBillingService mService;
+	private BillingClient mBillingClient;
 
 	// Flag that determines whether the store is ready to use.
 	private boolean bIsIapSetup;
@@ -77,6 +41,10 @@ public class GooglePlayStoreHelper implements StoreHelper
 
 	private final int UndefinedFailureResponse = -1;
 
+	private int outstandingConsumeRequests = 0;
+
+	private boolean bIsConsumeComplete = false;
+
 	private class GooglePlayProductDescription
 	{
 		// Product offer id 
@@ -91,6 +59,8 @@ public class GooglePlayStoreHelper implements StoreHelper
 		public Float priceRaw;
 		// Local currency code
 		public String currencyCode;
+		// Original Json representation
+		public String originalJson;
 	}
 
 	public interface PurchaseLaunchCallback
@@ -109,11 +79,46 @@ public class GooglePlayStoreHelper implements StoreHelper
 		gameActivity = InGameActivity;
 		productKey = InProductKey;
 
-		Intent serviceIntent = new Intent("com.android.vending.billing.InAppBillingService.BIND");
-		serviceIntent.setPackage("com.android.vending");
-		gameActivity.bindService(serviceIntent, mServiceConn, Context.BIND_AUTO_CREATE);
-	}
+		PurchasesUpdatedListener purchasesUpdatedListener = new PurchasesUpdatedListener(){
+			@Override
+			public void onPurchasesUpdated(BillingResult billingResult, List<Purchase> purchases){
+				if (billingResult.getResponseCode() == BillingClient.BillingResponseCode.OK
+					&& purchases != null) {
+						for (Purchase purchase : purchases) {
+							onPurchaseResult(billingResult, purchase);
+						}
+				} else
+				{ 
+					Log.debug("[GooglePlayStoreHelper] - GooglePlayStoreHelper::UserCancelled Purchase " + billingResult.getDebugMessage());
+					nativePurchaseComplete(billingResult.getResponseCode(), "", "", "", "");
+					// Handle an error caused by a user cancelling the purchase flow.
+				}
+			}
+		};
 
+		mBillingClient = BillingClient.newBuilder(gameActivity).setListener(purchasesUpdatedListener).enablePendingPurchases().build();
+		
+		mBillingClient.startConnection(new BillingClientStateListener() {
+			@Override
+			public void onBillingSetupFinished(BillingResult billingResult) {
+				if (billingResult.getResponseCode() == BillingClient.BillingResponseCode.OK) {
+					bIsIapSetup = true;
+					Log.debug("In-app billing supported for " + gameActivity.getPackageName());
+				}
+				else
+				{
+					Log.debug("In-app billing NOT supported for " + gameActivity.getPackageName() + " error " + billingResult.getResponseCode());
+				}
+			}
+    
+			@Override
+			public void onBillingServiceDisconnected() {
+				// Try to restart the connection on the next request to
+				// Google Play by calling the startConnection() method.
+				bIsIapSetup = false;
+			}
+		});
+	}
 
 	///////////////////////////////////////////////////////
 	// The StoreHelper interfaces implementation for Google Play Store.
@@ -144,187 +149,161 @@ public class GooglePlayStoreHelper implements StoreHelper
 				skuList.add(productId);
 			}
 
-			ArrayList<String> productIds = new ArrayList<String>();
-			ArrayList<String> titles = new ArrayList<String>();
-			ArrayList<String> descriptions = new ArrayList<String>();
-			ArrayList<String> prices = new ArrayList<String>();
-			ArrayList<Float> pricesRaw = new ArrayList<Float>();
-			ArrayList<String> currencyCodes = new ArrayList<String>();
+			Log.debug("[GooglePlayStoreHelper] - NumSkus: " + skuList.size());
 
-			while (skuList.size() > 0)
-			{
-				Log.debug("[GooglePlayStoreHelper] - NumSkus: " + skuList.size());
-				int numSkus = Math.min(QueryPurchaseMaxCount, skuList.size());
-				Log.debug("[GooglePlayStoreHelper] - Grabbing from 0 to " + numSkus);
-				List<String> skuSubList = skuList.subList(0, numSkus);
-				ArrayList<String> realSubList = new ArrayList<String>(skuSubList);
-				ArrayList<GooglePlayProductDescription> productDescriptions = new ArrayList<GooglePlayProductDescription>();
-
-				int response = QueryInAppPurchasesInternal(realSubList, productDescriptions);
-				if (response == BILLING_RESPONSE_RESULT_OK)
-				{
-					for (GooglePlayProductDescription product : productDescriptions)
+			SkuDetailsParams.Builder skuDetailsParams = SkuDetailsParams.newBuilder();
+			skuDetailsParams = skuDetailsParams.setSkusList(skuList).setType(BillingClient.SkuType.INAPP);
+			
+			mBillingClient.querySkuDetailsAsync(skuDetailsParams.build(),
+				new SkuDetailsResponseListener() {
+					@Override
+					public void onSkuDetailsResponse(BillingResult billingResult,
+						List<SkuDetails> skuDetailsList) 
 					{
-						productIds.add(product.id);
-						Log.debug("[GooglePlayStoreHelper] - GooglePlayStoreHelper::QueryInAppPurchases - Parsing details for: " + product.id);
+						// Process the result.
+						ArrayList<GooglePlayProductDescription> productDescriptions = new ArrayList<GooglePlayProductDescription>();
+						int response = billingResult.getResponseCode();
+						Log.debug("[GooglePlayStoreHelper] - GooglePlayStoreHelper::QueryInAppPurchases - Response " + response + " SkuDetails:" + skuDetailsList.toString());
+
+						// Should we send JSON, or somehow serialize this so we can have random access fields?
+						ArrayList<String> productIds = new ArrayList<String>();
+						ArrayList<String> titles = new ArrayList<String>();
+						ArrayList<String> descriptions = new ArrayList<String>();
+						ArrayList<String> prices = new ArrayList<String>();
+						ArrayList<Float> pricesRaw = new ArrayList<Float>();
+						ArrayList<String> currencyCodes = new ArrayList<String>();
+						ArrayList<String> originalJson = new ArrayList<String>();
+
+						if(response == BillingClient.BillingResponseCode.OK)
+						{
+							for (SkuDetails thisSku : skuDetailsList)
+							{
+								GooglePlayProductDescription newDescription = new GooglePlayProductDescription();
+								newDescription.id = thisSku.getSku();
+								newDescription.title = thisSku.getTitle();
+								newDescription.description = thisSku.getDescription();
+								newDescription.price = thisSku.getPrice();
+								double priceRaw = thisSku.getPriceAmountMicros() / 1000000.0;
+								newDescription.priceRaw = (float)priceRaw;
+								newDescription.currencyCode = thisSku.getPriceCurrencyCode();
+								newDescription.originalJson = thisSku.getOriginalJson();
+								productDescriptions.add(newDescription);
+							}
+						
+							for (GooglePlayProductDescription product : productDescriptions)
+							{
+								productIds.add(product.id);
+								Log.debug("[GooglePlayStoreHelper] - GooglePlayStoreHelper::QueryInAppPurchases - Parsing details for: " + product.id);
 
-						titles.add(product.title);
-						Log.debug("[GooglePlayStoreHelper] - title: " + product.title);
+								titles.add(product.title);
+								Log.debug("[GooglePlayStoreHelper] - title: " + product.title);
 
-						descriptions.add(product.description);
-						Log.debug("[GooglePlayStoreHelper] - description: " + product.description);
+								descriptions.add(product.description);
+								Log.debug("[GooglePlayStoreHelper] - description: " + product.description);
 
-						prices.add(product.price);
-						Log.debug("[GooglePlayStoreHelper] - price: " + product.price);
+								prices.add(product.price);
+								Log.debug("[GooglePlayStoreHelper] - price: " + product.price);
 
-						pricesRaw.add(product.priceRaw);
-						Log.debug("[GooglePlayStoreHelper] - price_amount_micros: " + product.priceRaw);
+								pricesRaw.add(product.priceRaw);
+								Log.debug("[GooglePlayStoreHelper] - price_amount_micros: " + product.priceRaw);
 
-						currencyCodes.add(product.currencyCode);
-						Log.debug("[GooglePlayStoreHelper] - price_currency_code: " + product.currencyCode);
-					}
-				}
-				else
-				{
-					Log.debug("[GooglePlayStoreHelper] - GooglePlayStoreHelper::QueryInAppPurchases - Failed with: " + response);
-					// If anything fails right now, stop immediately, not sure how to reconcile partial success/failure
-					nativeQueryComplete(response, null, null, null, null, null, null);
-					return false;
-				}
+								currencyCodes.add(product.currencyCode);
+								Log.debug("[GooglePlayStoreHelper] - price_currency_code: " + product.currencyCode);
 
-				skuSubList.clear();
-			}
+								originalJson.add(product.originalJson);
+								Log.debug("[GooglePlayStoreHelper] - original_json: " + product.originalJson);
 
-			float[] pricesRawPrimitive = new float[pricesRaw.size()];
-			for (int i = 0; i < pricesRaw.size(); i++)
-			{
-				pricesRawPrimitive[i] = pricesRaw.get(i);
-			}
+							}
+						}
+						else
+						{
+							Log.debug("[GooglePlayStoreHelper] - GooglePlayStoreHelper::QueryInAppPurchases - Failed with: " + response);
+							// If anything fails right now, stop immediately, not sure how to reconcile partial success/failure
+							nativeQueryComplete(response, null, null, null, null, null, null, null);
+						}
+						
+						float[] pricesRawPrimitive = new float[pricesRaw.size()];
+						for (int i = 0; i < pricesRaw.size(); i++)
+						{
+							pricesRawPrimitive[i] = pricesRaw.get(i);
+						}
 
-			Log.debug("[GooglePlayStoreHelper] - GooglePlayStoreHelper::QueryInAppPurchases " + productIds.size() + " items - Success!");
-			nativeQueryComplete(BILLING_RESPONSE_RESULT_OK, productIds.toArray(new String[productIds.size()]), titles.toArray(new String[titles.size()]), descriptions.toArray(new String[descriptions.size()]), prices.toArray(new String[prices.size()]), pricesRawPrimitive, currencyCodes.toArray(new String[currencyCodes.size()]));
-			Log.debug("[GooglePlayStoreHelper] - nativeQueryComplete done!");
+						Log.debug("[GooglePlayStoreHelper] - GooglePlayStoreHelper::QueryInAppPurchases " + productIds.size() + " items - Success!");
+						nativeQueryComplete(BillingClient.BillingResponseCode.OK, productIds.toArray(new String[productIds.size()]), titles.toArray(new String[titles.size()]), descriptions.toArray(new String[descriptions.size()]), prices.toArray(new String[prices.size()]), pricesRawPrimitive, currencyCodes.toArray(new String[currencyCodes.size()]), originalJson.toArray(new String[originalJson.size()]));
+						Log.debug("[GooglePlayStoreHelper] - nativeQueryComplete done!");
+					}
+				}
+			);			
 		}
 		else
 		{
 			// nothing to query
 			Log.debug("[GooglePlayStoreHelper] - no products given to query");
-			nativeQueryComplete(UndefinedFailureResponse, null, null, null, null, null, null);
+			nativeQueryComplete(UndefinedFailureResponse, null, null, null, null, null, null, null);
+			return false;
 		}
 
 		return true;
 	}
 
-	public int QueryInAppPurchasesInternal(ArrayList<String> skuList, ArrayList<GooglePlayProductDescription> outProducts)
-	{
-		Log.debug("[GooglePlayStoreHelper] - GooglePlayStoreHelper::QueryInAppPurchasesInternal");
-
-		Bundle querySkus = new Bundle();
-		querySkus.putStringArrayList(GET_SKU_DETAILS_ITEM_LIST, skuList);
-
-		int response = UndefinedFailureResponse;
-		try
-		{
-			Bundle skuDetails = mService.getSkuDetails(3, gameActivity.getPackageName(), ITEM_TYPE_INAPP, querySkus);
-
-			response = skuDetails.getInt(RESPONSE_CODE);
-			Log.debug("[GooglePlayStoreHelper] - GooglePlayStoreHelper::QueryInAppPurchases - Response " + response + " Bundle:" + skuDetails.toString());
-			if (response == BILLING_RESPONSE_RESULT_OK)
-			{
-				ArrayList<String> responseList = skuDetails.getStringArrayList(RESPONSE_GET_SKU_DETAILS_LIST);
-				for (String thisResponse : responseList)
-				{
-					JSONObject object = new JSONObject(thisResponse);
-
-					GooglePlayProductDescription newDescription = new GooglePlayProductDescription();
-					newDescription.id = object.getString("productId");
-					newDescription.title = object.getString("title");
-					newDescription.description = object.getString("description");
-					newDescription.price = object.getString("price");
-					double priceRaw = object.getDouble("price_amount_micros") / 1000000.0;
-					newDescription.priceRaw = (float)priceRaw;
-					newDescription.currencyCode = object.getString("price_currency_code");
-					outProducts.add(newDescription);
-				}
-			}
-		}
-		catch(Exception e)
-		{
-			Log.debug("[GooglePlayStoreHelper] - GooglePlayStoreHelper::QueryInAppPurchasesInternal - Failed! " + e.getMessage());
-		}
-
-		return response;
-	}
-
 	/**
 	 * Start the purchase flow for a particular sku
 	 */
 	final int purchaseIntentIdentifier = 1001;
 	public boolean BeginPurchase(String ProductID)
 	{
-		try
-		{
-			String devPayload = GenerateDevPayload(ProductID);
-			Bundle buyIntentBundle = null;
-			int response = -1;
-
-			if (gameActivity.IsInVRMode())
-			{
-				Bundle bundle = new Bundle();
-				bundle.putBoolean("vr", true);
-				response = mService.isBillingSupportedExtraParams(7, gameActivity.getPackageName(), ITEM_TYPE_INAPP, bundle);
-				if (response == BILLING_RESPONSE_RESULT_OK)
-				{
-					Log.debug("[GooglePlayStoreHelper] - GooglePlayStoreHelper::BeginPurchase - v7 VR purchase" + ProductID);
-					buyIntentBundle = mService.getBuyIntentExtraParams(7, gameActivity.getPackageName(), ProductID, ITEM_TYPE_INAPP, devPayload, bundle);
-				}
-				else
-				{
-					Log.debug("[GooglePlayStoreHelper] - GooglePlayStoreHelper::BeginPurchase - v3 IAB purchase:" + ProductID);
-					buyIntentBundle = mService.getBuyIntent(3, gameActivity.getPackageName(), ProductID, ITEM_TYPE_INAPP, devPayload);
-				}
-			}
-			else
-			{
-				Log.debug("[GooglePlayStoreHelper] - GooglePlayStoreHelper::BeginPurchase - v3 IAB purchase:" + ProductID);
-				buyIntentBundle = mService.getBuyIntent(3, gameActivity.getPackageName(), ProductID, ITEM_TYPE_INAPP, devPayload);
-			}
-
-			response = buyIntentBundle.getInt(RESPONSE_CODE);
-			if (response == BILLING_RESPONSE_RESULT_OK)
-			{
-				Log.debug("[GooglePlayStoreHelper] - GooglePlayStoreHelper::BeginPurchase - Starting Intent to buy " + ProductID);
-				PendingIntent pendingIntent = buyIntentBundle.getParcelable(RESPONSE_BUY_INTENT);
-				if (pendingIntent != null)
-				{
-					PurchaseLaunchCallback callback =  gameActivity.getPurchaseLaunchCallback();
-					if (callback != null)
+		ArrayList<String> skuList = new ArrayList<String>(1);
+		skuList.add(ProductID);
+
+		SkuDetailsParams.Builder skuDetailsParams = SkuDetailsParams.newBuilder();
+		skuDetailsParams.setSkusList(skuList).setType(BillingClient.SkuType.INAPP);
+		mBillingClient.querySkuDetailsAsync(skuDetailsParams.build(),
+				new SkuDetailsResponseListener() {
+					@Override
+					public void onSkuDetailsResponse(BillingResult billingResult,
+						List<SkuDetails> skuDetailsList) 
 					{
-						callback.launchForResult(pendingIntent, purchaseIntentIdentifier);
-					}
-					else
-					{
-						gameActivity.startIntentSenderForResult(pendingIntent.getIntentSender(), purchaseIntentIdentifier, new Intent(), Integer.valueOf(0), Integer.valueOf(0), Integer.valueOf(0));
-					}
-				}
-				else
-				{
-					Log.debug("[GooglePlayStoreHelper] - GooglePlayStoreHelper::BeginPurchase - pendingIntent was null, possible non consumed item");
-					nativePurchaseComplete(UndefinedFailureResponse, ProductID, "", "", "");
-				}
-			}
-			else
-			{
-				Log.debug("[GooglePlayStoreHelper] - GooglePlayStoreHelper::BeginPurchase - Failed with error: " + TranslateServerResponseCode(response));
-				nativePurchaseComplete(response, ProductID, "", "", "");
-			}
-		}
-		catch(Exception e)
-		{
-			Log.debug("[GooglePlayStoreHelper] - GooglePlayStoreHelper::BeginPurchase - Failed! " + e.getMessage());
-			nativePurchaseComplete(UndefinedFailureResponse, ProductID, "", "", "");
-		}
+						
+						int response = billingResult.getResponseCode();
+						if(response == BillingClient.BillingResponseCode.OK)
+						{
+							for(SkuDetails skuDetails : skuDetailsList)
+							{
+								BillingFlowParams.Builder flowParams = BillingFlowParams.newBuilder().setSkuDetails(skuDetails);
+			
+			
+								if (gameActivity.IsInVRMode())
+								{
+									int responseCode = mBillingClient.isFeatureSupported(BillingClient.FeatureType.IN_APP_ITEMS_ON_VR).getResponseCode();
+									boolean isSupportedInVR = (responseCode == BillingClient.BillingResponseCode.OK);
+									if (isSupportedInVR)
+									{
+										Log.debug("[GooglePlayStoreHelper] - GooglePlayStoreHelper::BeginPurchase - v7 VR purchase" + skuDetails.getSku());
+										flowParams = flowParams.setVrPurchaseFlow(true);
+									}
+									else
+									{
+										Log.debug("[GooglePlayStoreHelper] - GooglePlayStoreHelper::BeginPurchase - v3 IAB purchase:" + skuDetails.getSku());
+										flowParams = flowParams.setVrPurchaseFlow(false);
+									}
+								}
+								else
+								{
+									Log.debug("[GooglePlayStoreHelper] - GooglePlayStoreHelper::BeginPurchase - v3 IAB purchase:" + skuDetails.getSku());
+									flowParams = flowParams.setVrPurchaseFlow(false);
+								}
 
+								int responseCode = mBillingClient.launchBillingFlow(gameActivity, flowParams.build()).getResponseCode();
+								Log.debug("[GooglePlayStoreHelper] - GooglePlayStoreHelper::BeginPurchase - Launching billing flow " + skuDetails.getSku());
+							}
+						}
+						else
+						{
+							Log.debug("[GooglePlayStoreHelper] - GooglePlayStoreHelper::BeginPurchase - Failed! " + TranslateServerResponseCode(response));
+							nativePurchaseComplete(UndefinedFailureResponse, "", "", "", "");
+						}
+					}
+		});
 		return true;
 	}
 
@@ -334,99 +313,113 @@ public class GooglePlayStoreHelper implements StoreHelper
 	public boolean RestorePurchases(String[] InProductIDs, boolean[] bConsumable)
 	{
 		Log.debug("[GooglePlayStoreHelper] - GooglePlayStoreHelper::RestorePurchases");
-		ArrayList<String> ownedSkus = new ArrayList<String>();
-		ArrayList<String> purchaseDataList = new ArrayList<String>();
-		ArrayList<String> signatureList = new ArrayList<String>();
 
-		// On first pass the continuation token should be null.
-		// This will allow us to gather large sets of purchased items recursively
-		int responseCode = GatherOwnedPurchaseData(ownedSkus, purchaseDataList, signatureList, null);
-		if (responseCode == BILLING_RESPONSE_RESULT_OK)
+		ArrayList<Purchase> ownedProducts = new ArrayList<Purchase>();
+		int responseCode = GatherOwnedPurchaseData(ownedProducts);
+		if (responseCode == BillingClient.BillingResponseCode.OK)
 		{
-			Log.debug("[GooglePlayStoreHelper] - GooglePlayStoreHelper::RestorePurchases - User has previously purchased " + ownedSkus.size() + " inapp products" );
+			Log.debug("[GooglePlayStoreHelper] - GooglePlayStoreHelper::RestorePurchases - User has previously purchased " + ownedProducts.size() + " inapp products" );
+
+			final ArrayList<Purchase> f_ownedProducts = ownedProducts;
 
-			final ArrayList<String> f_ownedSkus = ownedSkus;
-			final ArrayList<String> f_purchaseDataList = purchaseDataList;
-			final ArrayList<String> f_signatureList = signatureList;
+			final ArrayList<String> f_ownedSkus = new ArrayList<String>();
+			final ArrayList<String> f_purchaseDataList = new ArrayList<String>();
+			final ArrayList<String> f_signatureList = new ArrayList<String>();
 			final String[] RestoreProductIDs = InProductIDs;
 			final boolean[] bRestoreConsumableFlags = bConsumable;
 
+			// Is this here to ensure callbacks are fired on the main thread?
 			Handler mainHandler = new Handler(Looper.getMainLooper());
 			mainHandler.post(new Runnable()
 			{
 				@Override
 				public void run()
 				{
-					try
+					Log.debug("[GooglePlayStoreHelper] - GooglePlayStoreHelper::RestorePurchases - Now consuming any purchases that may have been missed.");
+					final ArrayList<String> productTokens = new ArrayList<String>();
+					final ArrayList<String> receipts = new ArrayList<String>();
+					final ArrayList<Integer> cachedErrorResponses = new ArrayList<Integer>();
+
+					bIsConsumeComplete = false;
+					// @todo possible bug, restores ALL ownedSkus, not just those provided by caller in RestoreProductIDs
+					for (Purchase ownedProduct : f_ownedProducts)
 					{
-						Log.debug("[GooglePlayStoreHelper] - GooglePlayStoreHelper::RestorePurchases - Now consuming any purchases that may have been missed.");
-						ArrayList<String> productTokens = new ArrayList<String>();
-						ArrayList<String> receipts = new ArrayList<String>();
-						int cachedResponse = 0;
-
-						// @todo possible bug, restores ALL ownedSkus, not just those provided by caller in RestoreProductIDs
-						for (int Idx = 0; Idx < f_ownedSkus.size(); Idx++)
-						{
-							String purchaseData = f_purchaseDataList.get(Idx);
-							String dataSignature = f_signatureList.get(Idx);
-
-							try
-							{
-								Purchase purchase = new Purchase(ITEM_TYPE_INAPP, purchaseData, dataSignature);
-								productTokens.add(purchase.getToken());
+						
+						productTokens.add(ownedProduct.getPurchaseToken());
 
-								boolean bTryToConsume = false;
-								int consumeResponse = 0;
+						boolean bTryToConsume = false;
 
-								// This is assuming that all purchases should be consumed. Consuming a purchase that is meant to be a one-time purchase makes it so the
-								// user is able to buy it again. Also, it makes it so the purchase will not be able to be restored again in the future.
+						// This is assuming that all purchases should be consumed. Consuming a purchase that is meant to be a one-time purchase makes it so the
+						// user is able to buy it again. Also, it makes it so the purchase will not be able to be restored again in the future.
 
-								for (int Idy = 0; Idy < RestoreProductIDs.length; Idy++)
+						for (int Idy = 0; Idy < RestoreProductIDs.length; Idy++)
+						{
+							if( ownedProduct.getSku().equals(RestoreProductIDs[Idy]) )
+							{
+								if(Idy < bRestoreConsumableFlags.length)
 								{
-									if( purchase.getSku().equals(RestoreProductIDs[Idy]) )
-									{
-										if(Idy < bRestoreConsumableFlags.length)
-										{
-											bTryToConsume = bRestoreConsumableFlags[Idy];
-											Log.debug("[GooglePlayStoreHelper] - GooglePlayStoreHelper::RestorePurchases - Found Consumable Flag for Product " + purchase.getSku() + " bConsumable = " + bTryToConsume);
-										}
-										break;
-									}
+									bTryToConsume = bRestoreConsumableFlags[Idy];
+									Log.debug("[GooglePlayStoreHelper] - GooglePlayStoreHelper::RestorePurchases - Found Consumable Flag for Product " + ownedProduct.getSku() + " bConsumable = " + bTryToConsume);
 								}
+								break;
+							}
+						}
 
-								if (bTryToConsume)
-								{
-									Log.debug("[GooglePlayStoreHelper] - GooglePlayStoreHelper::RestorePurchases - Attempting to consume " + purchase.getSku());
-									consumeResponse = mService.consumePurchase(3, gameActivity.getPackageName(), purchase.getToken());
-								}
+						final Purchase constProduct = ownedProduct;
+						if (bTryToConsume)
+						{
+							Log.debug("[GooglePlayStoreHelper] - GooglePlayStoreHelper::RestorePurchases - Attempting to consume " + ownedProduct.getSku());
+							ConsumeParams consumeParams = ConsumeParams.newBuilder().setPurchaseToken(ownedProduct.getPurchaseToken()).build();
+							mBillingClient.consumeAsync(consumeParams, (new ConsumeResponseListener() {
+								@Override
+								public void onConsumeResponse(BillingResult billingResult, String purchaseToken) {
+									if(billingResult.getResponseCode() == BillingClient.BillingResponseCode.OK)
+									{
+										// How do we get purchase here.. need to figure out a way to persist cachedResponse and receipts
+										Log.debug("[GooglePlayStoreHelper] - GooglePlayStoreHelper::RestorePurchases - Purchase restored for " + constProduct.getSku());
+										String receipt = Base64.encode(constProduct.getOriginalJson().getBytes());
+										receipts.add(receipt);
+										
+										f_ownedSkus.add(constProduct.getSku());
+										f_signatureList.add(constProduct.getSignature());
+									}
+									else
+									{
+										Log.debug("[GooglePlayStoreHelper] - GooglePlayStoreHelper::RestorePurchases - consumePurchase failed for " + constProduct.getSku() + " with error " + billingResult.getResponseCode());
+										receipts.add("");
+										cachedErrorResponses.add(billingResult.getResponseCode());
+									}
+									outstandingConsumeRequests--;
 
-								if (consumeResponse == 0)
-								{
-									Log.debug("[GooglePlayStoreHelper] - GooglePlayStoreHelper::RestorePurchases - Purchase restored for " + purchase.getSku());
-									String receipt = Base64.encode(purchase.getOriginalJson().getBytes());
-									receipts.add(receipt);
-								}
-								else
-								{
-									Log.debug("[GooglePlayStoreHelper] - GooglePlayStoreHelper::RestorePurchases - consumePurchase failed for " + purchase.getSku() + " with error " + consumeResponse);
-									receipts.add("");
-									cachedResponse = consumeResponse;
+									if(outstandingConsumeRequests <= 0)
+									{
+										Log.debug("[GooglePlayStoreHelper] - GooglePlayStoreHelper::RestorePurchases - Success!");
+										int finalResponse = (cachedErrorResponses.size() > 0) ? cachedErrorResponses.get(0) : BillingClient.BillingResponseCode.OK;
+										nativeRestorePurchasesComplete(finalResponse, f_ownedSkus.toArray(new String[f_ownedSkus.size()]), productTokens.toArray(new String[productTokens.size()]), receipts.toArray(new String[receipts.size()]), f_signatureList.toArray(new String[f_signatureList.size()]));
+										bIsConsumeComplete = true;
+									}
 								}
-							}
-							catch (JSONException e)
-							{
-								Log.debug("[GooglePlayStoreHelper] - GooglePlayStoreHelper::RestorePurchases - Failed to parse receipt! " + e.getMessage());
-								productTokens.add("");
-								receipts.add("");
-							}
+							}));
+								
+							outstandingConsumeRequests++;
+						}
+						else
+						{
+							// How do we get purchase here.. need to figure out a way to persist cachedResponse and receipts
+							Log.debug("[GooglePlayStoreHelper] - GooglePlayStoreHelper::RestorePurchases - Purchase restored for " + constProduct.getSku());
+							String receipt = Base64.encode(constProduct.getOriginalJson().getBytes());
+							receipts.add(receipt);
+										
+							f_ownedSkus.add(constProduct.getSku());
+							f_signatureList.add(constProduct.getSignature());
 						}
-						Log.debug("[GooglePlayStoreHelper] - GooglePlayStoreHelper::RestorePurchases - Success!");
-						nativeRestorePurchasesComplete(cachedResponse, f_ownedSkus.toArray(new String[f_ownedSkus.size()]), productTokens.toArray(new String[productTokens.size()]), receipts.toArray(new String[receipts.size()]), f_signatureList.toArray(new String[f_signatureList.size()]));
 					}
-					catch (Exception e)
+
+					// There was nothing to consume, so we are done!
+					if(outstandingConsumeRequests <= 0 && !bIsConsumeComplete)
 					{
-						Log.debug("[GooglePlayStoreHelper] - GooglePlayStoreHelper::RestorePurchases - consumePurchase failed. " + e.getMessage());
-						nativeRestorePurchasesComplete(UndefinedFailureResponse, null, null, null, null);
+						Log.debug("[GooglePlayStoreHelper] - GooglePlayStoreHelper::RestorePurchases - Success!");
+						nativeRestorePurchasesComplete(BillingClient.BillingResponseCode.OK, f_ownedSkus.toArray(new String[f_ownedSkus.size()]), productTokens.toArray(new String[productTokens.size()]), receipts.toArray(new String[receipts.size()]), f_signatureList.toArray(new String[f_signatureList.size()]));
 					}
 				}
 			});
@@ -448,35 +441,23 @@ public class GooglePlayStoreHelper implements StoreHelper
 		ArrayList<String> purchaseDataList = new ArrayList<String>();
 		ArrayList<String> signatureList = new ArrayList<String>();
 
-		// On first pass the continuation token should be null.
-		// This will allow us to gather large sets of purchased items recursively
-		int responseCode = GatherOwnedPurchaseData(ownedSkus, purchaseDataList, signatureList, null);
-		if (responseCode == BILLING_RESPONSE_RESULT_OK)
+		ArrayList<Purchase> ownedProducts = new ArrayList<Purchase>();
+
+		int responseCode = GatherOwnedPurchaseData(ownedProducts);
+		if (responseCode == BillingClient.BillingResponseCode.OK)
 		{
 			Log.debug("[GooglePlayStoreHelper] - GooglePlayStoreHelper::QueryExistingPurchases - User has previously purchased " + ownedSkus.size() + " inapp products" );
 
 			ArrayList<String> productTokens = new ArrayList<String>();
 			ArrayList<String> receipts = new ArrayList<String>();
 
-			for (int Idx = 0; Idx < ownedSkus.size(); Idx++)
+			for (Purchase purchase : ownedProducts)
 			{
-				String purchaseData = purchaseDataList.get(Idx);
-				String dataSignature = signatureList.get(Idx);
-
-				try
-				{
-					Purchase purchase = new Purchase(ITEM_TYPE_INAPP, purchaseData, dataSignature);
-					productTokens.add(purchase.getToken());
-
-					String receipt = Base64.encode(purchase.getOriginalJson().getBytes());
-					receipts.add(receipt);
-				}
-				catch (JSONException e)
-				{
-					Log.debug("[GooglePlayStoreHelper] - GooglePlayStoreHelper::QueryExistingPurchases - Failed to parse receipt! " + e.getMessage());
-					productTokens.add("");
-					receipts.add("");
-				}
+				productTokens.add(purchase.getPurchaseToken());
+				ownedSkus.add(purchase.getSku());
+				signatureList.add(purchase.getSignature());
+				String receipt = Base64.encode(purchase.getOriginalJson().getBytes());
+				receipts.add(receipt);
 			}
 
 			Log.debug("[GooglePlayStoreHelper] - GooglePlayStoreHelper::QueryExistingPurchases - Success!");
@@ -503,71 +484,24 @@ public class GooglePlayStoreHelper implements StoreHelper
 			@Override
 			public void run()
 			{
-				try
-				{
-					Log.debug("[GooglePlayStoreHelper] - Consuming token: " + f_purchaseToken);
-					int consumeResponse = mService.consumePurchase(3, gameActivity.getPackageName(), f_purchaseToken);
-					if (consumeResponse == BILLING_RESPONSE_RESULT_OK)
-					{
-						Log.debug("[GooglePlayStoreHelper] - ConsumePurchase success");
-					}
-					else
-					{
-						Log.debug("[GooglePlayStoreHelper] - ConsumePurchase failed with error " + TranslateServerResponseCode(consumeResponse));
+				ConsumeParams consumeParams = ConsumeParams.newBuilder().setPurchaseToken(f_purchaseToken).build();
+				mBillingClient.consumeAsync(consumeParams, (new ConsumeResponseListener() {
+					@Override
+					public void onConsumeResponse(BillingResult billingResult, String purchaseToken) {
+						if(billingResult.getResponseCode() == BillingClient.BillingResponseCode.OK)
+						{
+							Log.debug("[GooglePlayStoreHelper] - ConsumePurchase success");
+						}
+						else
+						{
+							Log.debug("[GooglePlayStoreHelper] - ConsumePurchase failed with error " + TranslateServerResponseCode(billingResult.getResponseCode()));
+						}
 					}
-				}
-				catch(Exception e)
-				{
-					Log.debug("[GooglePlayStoreHelper] - ConsumePurchase failed. " + e.getMessage());
-				}
+				}));
 			}
 		});
 	}
 
-	///////////////////////////////////////////////////////
-	// IInAppBillingService service registration and unregistration
-
-	private ServiceConnection mServiceConn = new ServiceConnection()
-	{
-		@Override
-		public void onServiceDisconnected(ComponentName name)
-		{
-			Log.debug("[GooglePlayStoreHelper] - ServiceConnection::onServiceDisconnected");
-			mService = null;
-			bIsIapSetup = false;
-		}
-
-		@Override
-		public void onServiceConnected(ComponentName name, IBinder service)
-		{
-			Log.debug("[GooglePlayStoreHelper] - ServiceConnection::onServiceConnected");
-			mService = IInAppBillingService.Stub.asInterface(service);
-			bIsIapSetup = true;
-
-			try
-			{
-                Log.debug("Checking for in-app billing 3 support.");
-
-                // check for in-app billing v3 support
-                int response = mService.isBillingSupported(3, gameActivity.getPackageName(), ITEM_TYPE_INAPP);
-                if (response != BILLING_RESPONSE_RESULT_OK)
-				{
-					Log.debug("In-app billing version 3 NOT supported for " + gameActivity.getPackageName() + " error " + response);
-                }
-				else
-				{
-					Log.debug("In-app billing version 3 supported for " + gameActivity.getPackageName());
-				}
-			}
-            catch (RemoteException e)
-			{
-				e.printStackTrace();
-				return;
-            }
-		}
-	};
-
-
 	///////////////////////////////////////////////////////
 	// Game Activity/Context driven methods we need to listen for.
 
@@ -578,9 +512,9 @@ public class GooglePlayStoreHelper implements StoreHelper
 	{
 		Log.debug("[GooglePlayStoreHelper] - GooglePlayStoreHelper::onDestroy");
 
-		if (mService != null)
+		if (mBillingClient != null)
 		{
-			gameActivity.unbindService(mServiceConn);
+			mBillingClient = null; //WMM ???
 		}
 	}
 
@@ -588,141 +522,97 @@ public class GooglePlayStoreHelper implements StoreHelper
 	 * Route taken by the Purchase workflow. We listen for our purchaseIntentIdentifier request code and
 	 * handle the response accordingly
 	 */
-	public boolean onActivityResult(int requestCode, int resultCode, Intent data)
+	public boolean onPurchaseResult(BillingResult billingResult, Purchase purchase)
 	{
-		Log.debug("[GooglePlayStoreHelper] - GooglePlayStoreHelper::onActivityResult");
+		Log.debug("[GooglePlayStoreHelper] - GooglePlayStoreHelper::onPurchaseResult");
 
-		if (requestCode == purchaseIntentIdentifier)
+		if (billingResult == null)
 		{
-		    if (data == null)
-			{
-				Log.debug("Null data in purchase activity result.");
-				nativePurchaseComplete(UndefinedFailureResponse, "", "", "", "");
-				return true;
-			}
+			Log.debug("Null data in purchase activity result.");
+			nativePurchaseComplete(UndefinedFailureResponse, "", "", "", "");
+			return true;
+		}
 
-			int responseCode = data.getIntExtra(RESPONSE_CODE, 0);
-			if (resultCode == Activity.RESULT_OK)
-			{
-				final String purchaseData = data.getStringExtra(RESPONSE_INAPP_PURCHASE_DATA);
-				final String dataSignature = data.getStringExtra(RESPONSE_INAPP_SIGNATURE);
+		int responseCode = billingResult.getResponseCode();
+		if (responseCode == BillingClient.BillingResponseCode.OK)
+		{
+			Log.debug("[GooglePlayStoreHelper] - GooglePlayStoreHelper::onActivityResult - Processing purchase result. Response Code: " + TranslateServerResponseCode(responseCode));
+			Log.debug("Purchase data: " + purchase.toString());
+			Log.debug("Data signature: " + purchase.getSignature());
 
-				Log.debug("[GooglePlayStoreHelper] - GooglePlayStoreHelper::onActivityResult - Processing purchase result. Response Code: " + TranslateServerResponseCode(responseCode));
-				if (responseCode == BILLING_RESPONSE_RESULT_OK)
+			if(purchase.getPurchaseState() == Purchase.PurchaseState.PURCHASED)
+			{
+				final String sku = purchase.getSku();
+				final Purchase f_purchase = purchase;
+				Handler mainHandler = new Handler(Looper.getMainLooper());
+				mainHandler.post(new Runnable()
 				{
-					Log.debug("Purchase data: " + purchaseData);
-					Log.debug("Data signature: " + dataSignature);
-					Log.debug("Extras: " + data.getExtras());
-
-					if (purchaseData == null || dataSignature == null)
-					{
-						Log.debug("[GooglePlayStoreHelper] - GooglePlayStoreHelper::onActivityResult - Either purchaseData or dataSignature is null");
-						nativePurchaseComplete(UndefinedFailureResponse, "", "", "", "");
-					}
-
-					try
+					@Override
+					public void run()
 					{
-						final Purchase purchase = new Purchase(ITEM_TYPE_INAPP, purchaseData, dataSignature);
-						final String sku = purchase.getSku();
-						String developerPayload = purchase.getDeveloperPayload();
-						if (VerifyPayload(developerPayload, sku))
-						{
-							Handler mainHandler = new Handler(Looper.getMainLooper());
-							mainHandler.post(new Runnable()
-							{
-								@Override
-								public void run()
-								{
-									String receipt = Base64.encode(purchase.getOriginalJson().getBytes());
-									nativePurchaseComplete(BILLING_RESPONSE_RESULT_OK, sku, purchase.getToken(), receipt, purchase.getSignature());
-								}
-							});
-						}
-						else
-						{
-							Log.debug("[GooglePlayStoreHelper] - GooglePlayStoreHelper::onActivityResult - Failed for verify developer payload for " + sku);
-							nativePurchaseComplete(UndefinedFailureResponse, sku, "", "", "");
-						}
-					}
-					catch (JSONException e)
-					{
-						Log.debug("[GooglePlayStoreHelper] - GooglePlayStoreHelper::onActivityResult - Failed for purchase request! " + e.getMessage());
-						nativePurchaseComplete(UndefinedFailureResponse, "", "", "", "");
+						String receipt = Base64.encode(f_purchase.getOriginalJson().getBytes());
+						nativePurchaseComplete(BillingClient.BillingResponseCode.OK, sku, f_purchase.getPurchaseToken(), receipt, f_purchase.getSignature());
 					}
-				}
-				else
-				{
-					Log.debug("[GooglePlayStoreHelper] - GooglePlayStoreHelper::onActivityResult - Failed for purchase request!. " + TranslateServerResponseCode(responseCode));
-					nativePurchaseComplete(responseCode, "", "", "", "");
-				}
-			}
-			else if (resultCode == Activity.RESULT_CANCELED)
-			{
-				Log.debug("[GooglePlayStoreHelper] - GooglePlayStoreHelper::onActivityResult - Purchase canceled." + TranslateServerResponseCode(responseCode));
-				nativePurchaseComplete(BILLING_RESPONSE_RESULT_USER_CANCELED, "", "", "", "");
+				});
 			}
 			else
 			{
-				Log.debug("[GooglePlayStoreHelper] - GooglePlayStoreHelper::onActivityResult - Purchase failed. Result code: " +
-					Integer.toString(resultCode) + ". Response: " + TranslateServerResponseCode(responseCode));
-				nativePurchaseComplete(UndefinedFailureResponse, "", "", "", "");
+				// This is a pending purchase.. need to hold off on reporting as purchase to game until it actually occurs.
+				// Special case - if pending purchase is USER_CANCELLED but the SKU is returned, then it is deferred
+				final String sku = purchase.getSku();
+				nativePurchaseComplete(BillingClient.BillingResponseCode.USER_CANCELED, sku, "", "", "");
 			}
-
-			return true;
+		}
+		else if (responseCode == BillingClient.BillingResponseCode.USER_CANCELED)
+		{
+			Log.debug("[GooglePlayStoreHelper] - GooglePlayStoreHelper::onActivityResult - Purchase canceled." + TranslateServerResponseCode(responseCode));
+			nativePurchaseComplete(BillingClient.BillingResponseCode.USER_CANCELED, "", "", "", "");
+		}
+		else
+		{
+			Log.debug("[GooglePlayStoreHelper] - GooglePlayStoreHelper::onActivityResult - Purchase failed. Result code: " +
+				Integer.toString(responseCode) + ". Response: " + TranslateServerResponseCode(responseCode));
+			nativePurchaseComplete(UndefinedFailureResponse, "", "", "", "");
 		}
 
-		return false;
+		return true;
 	}
 
 
 	///////////////////////////////////////////////////////
 	// Internal helper functions that deal assist with various IAB related events
 
-	/**
-	 * Create a UE4 specific unique string that will be used to verify purchases are legit.
-	 */
-	private String GenerateDevPayload( String ProductId )
-	{
-		return "ue4." + ProductId;
-	}
-
-	/**
-	 * Check the returned payload matches one for the product we are buying.
-	 */
-	private boolean VerifyPayload( String ExistingPayload, String ProductId )
-	{
-		String GeneratedPayload = GenerateDevPayload(ProductId);
-		Log.debug("[GooglePlayStoreHelper] - GooglePlayStoreHelper::VerifyPayload, ExistingPayload: " + ExistingPayload);
-		Log.debug("[GooglePlayStoreHelper] - GooglePlayStoreHelper::VerifyPayload, GeneratedPayload: " + GeneratedPayload);
-
-		return ExistingPayload.equals(GeneratedPayload);
-	}
-
 	/**
 	 * Get a text tranlation of the Response Codes returned by google play.
 	 */
 	private String TranslateServerResponseCode(final int serverResponseCode)
 	{
-		// taken from http://developer.android.com/google/play/billing/billing_reference.html
+		// taken from https://developer.android.com/reference/com/android/billingclient/api/BillingClient.BillingResponseCode
 		switch(serverResponseCode)
 		{
-			case BILLING_RESPONSE_RESULT_OK:
+			case BillingClient.BillingResponseCode.OK:
 				return "Success";
-			case BILLING_RESPONSE_RESULT_USER_CANCELED:
+			case BillingClient.BillingResponseCode.USER_CANCELED:
 				return "User pressed back or canceled a dialog";
-			case BILLING_RESPONSE_RESULT_SERVICE_UNAVAILABLE:
+			case BillingClient.BillingResponseCode.SERVICE_UNAVAILABLE:
 				return "Network connection is down";
-			case BILLING_RESPONSE_RESULT_BILLING_UNAVAILABLE:
+			case BillingClient.BillingResponseCode.SERVICE_TIMEOUT:
+				return "The request has reached the maximum timeout before Google Play responds";
+			case BillingClient.BillingResponseCode.SERVICE_DISCONNECTED:
+				return "Play Store service is not connected now - potentially transient state";
+			case BillingClient.BillingResponseCode.BILLING_UNAVAILABLE:
 				return "Billing API version is not supported for the type requested";
-			case BILLING_RESPONSE_RESULT_ITEM_UNAVAILABLE:
+			case BillingClient.BillingResponseCode.FEATURE_NOT_SUPPORTED:
+				return "Requested feature is not supported by Play Store on the current device";
+			case BillingClient.BillingResponseCode.ITEM_UNAVAILABLE:
 				return "Requested product is not available for purchase";
-			case BILLING_RESPONSE_RESULT_DEVELOPER_ERROR:
+			case BillingClient.BillingResponseCode.DEVELOPER_ERROR:
 				return "Invalid arguments provided to the API. This error can also indicate that the application was not correctly signed or properly set up for In-app Billing in Google Play, or does not have the necessary permissions in its manifest";
-			case BILLING_RESPONSE_RESULT_ERROR:
+			case BillingClient.BillingResponseCode.ERROR:
 				return "Fatal error during the API action";
-			case BILLING_RESPONSE_RESULT_ITEM_ALREADY_OWNED:
+			case BillingClient.BillingResponseCode.ITEM_ALREADY_OWNED:
 				return "Failure to purchase since item is already owned";
-			case BILLING_RESPONSE_RESULT_ITEM_NOT_OWNED:
+			case BillingClient.BillingResponseCode.ITEM_NOT_OWNED:
 				return "Failure to consume since item is not owned";
 			default:
 				return "Unknown Server Response Code";
@@ -733,39 +623,25 @@ public class GooglePlayStoreHelper implements StoreHelper
 	 * Recursive functionality to gather all of the purchases owned by a user.
 	 * if the user owns a lot of products then we need to getPurchases again with a continuationToken
 	 */
-	private int GatherOwnedPurchaseData(ArrayList<String> inOwnedSkus, ArrayList<String> inPurchaseDataList, ArrayList<String> inSignatureList, String inContinuationToken)
+	private int GatherOwnedPurchaseData(ArrayList<Purchase> inPurchases)
 	{
 		int responseCode = UndefinedFailureResponse;
 		try
 		{
-			Bundle ownedItems = mService.getPurchases(3, gameActivity.getPackageName(), ITEM_TYPE_INAPP, inContinuationToken);
+			Purchase.PurchasesResult purchasesResult = mBillingClient.queryPurchases(BillingClient.SkuType.INAPP);
 
-			responseCode = ownedItems.getInt(RESPONSE_CODE);
+			responseCode = purchasesResult.getResponseCode();
 			Log.debug("[GooglePlayStoreHelper] - GooglePlayStoreHelper::GatherOwnedPurchaseData - getPurchases result. Response Code: " + TranslateServerResponseCode(responseCode));
-			if (responseCode == BILLING_RESPONSE_RESULT_OK)
+			if (responseCode == BillingClient.BillingResponseCode.OK)
 			{
-				ArrayList<String> ownedSkus = ownedItems.getStringArrayList(RESPONSE_INAPP_ITEM_LIST);
-				ArrayList<String> purchaseDataList = ownedItems.getStringArrayList(RESPONSE_INAPP_PURCHASE_DATA_LIST);
-				ArrayList<String> signatureList = ownedItems.getStringArrayList(RESPONSE_INAPP_SIGNATURE_LIST);
-				String continuationToken = ownedItems.getString(INAPP_CONTINUATION_TOKEN);
-
-				for (int Idx = 0; Idx < purchaseDataList.size(); Idx++)
+				List<Purchase> ownedItems = purchasesResult.getPurchasesList();
+				for (Purchase purchase : ownedItems)
 				{
-					String sku = ownedSkus.get(Idx);
-					inOwnedSkus.add(sku);
-
-					String purchaseData = purchaseDataList.get(Idx);
-					inPurchaseDataList.add(purchaseData);
-
-					String signature = signatureList.get(Idx);
-					inSignatureList.add(signature);
-				}
-
-				// if continuationToken != null, call getPurchases again
-				// and pass in the token to retrieve more items
-				if (continuationToken != null)
-				{
-					return GatherOwnedPurchaseData(inOwnedSkus, inPurchaseDataList, inSignatureList, inContinuationToken);
+					// Don't Mark pending purchases as owned
+					if(purchase.getPurchaseState() == Purchase.PurchaseState.PURCHASED)
+					{
+						inPurchases.add(purchase);
+					}
 				}
 			}
 		}
@@ -778,7 +654,7 @@ public class GooglePlayStoreHelper implements StoreHelper
 	}
 
 	// Callback that notify the C++ implementation that a task has completed
-	public native void nativeQueryComplete(int responseCode, String[] productIDs, String[] titles, String[] descriptions, String[] prices, float[] pricesRaw, String[] currencyCodes );
+	public native void nativeQueryComplete(int responseCode, String[] productIDs, String[] titles, String[] descriptions, String[] prices, float[] pricesRaw, String[] currencyCodes, String[] originalJson );
 	public native void nativePurchaseComplete(int responseCode, String ProductId, String ProductToken, String ReceiptData, String Signature);
 	public native void nativeRestorePurchasesComplete(int responseCode, String[] ProductIds, String[] ProductTokens, String[] ReceiptsData, String[] Signatures);
 	public native void nativeQueryExistingPurchasesComplete(int responseCode, String[] ProductIds, String[] ProductTokens, String[] ReceiptsData, String[] Signatures);
diff --git a/Engine/Build/Android/Java/src/com/epicgames/ue4/StoreHelper.java b/Engine/Build/Android/Java/src/com/epicgames/ue4/StoreHelper.java
--- a/Engine/Build/Android/Java/src/com/epicgames/ue4/StoreHelper.java
+++ b/Engine/Build/Android/Java/src/com/epicgames/ue4/StoreHelper.java
@@ -14,5 +14,4 @@ public interface StoreHelper
 	public boolean QueryExistingPurchases();
 	public boolean RestorePurchases(String[] InProductIDs, boolean[] bConsumable);
 	public void	onDestroy();
-	public boolean onActivityResult(int requestCode, int resultCode, Intent data);
 }
\ No newline at end of file
diff --git a/Engine/Config/Android/AndroidEngine.ini b/Engine/Config/Android/AndroidEngine.ini
--- a/Engine/Config/Android/AndroidEngine.ini
+++ b/Engine/Config/Android/AndroidEngine.ini
@@ -78,3 +78,7 @@ TaskGraph.EnablePowerSavingThreadPriorityReduction=1
 r.Vulkan.Depth24Bit=1
 ;GalaxyS6Edge=577
 r.Shadow.ForceSerialSingleRenderPass=1
+
+[OnlineSubsystemGooglePlay.Store]
+bSupportsInAppPurchasing=true
+bUseStoreV2=true
diff --git a/Engine/Plugins/Online/Android/OnlineSubsystemGooglePlay/Source/OnlineSubsystemGooglePlay_UPL.xml b/Engine/Plugins/Online/Android/OnlineSubsystemGooglePlay/Source/OnlineSubsystemGooglePlay_UPL.xml
--- a/Engine/Plugins/Online/Android/OnlineSubsystemGooglePlay/Source/OnlineSubsystemGooglePlay_UPL.xml
+++ b/Engine/Plugins/Online/Android/OnlineSubsystemGooglePlay/Source/OnlineSubsystemGooglePlay_UPL.xml
@@ -5,7 +5,16 @@
 	<init>
 	  <log text="Subsystem Google Play SDK Android init"/>
 	</init>
-
+	
+	<!-- gradle dependencies additions -->
+	<buildGradleAdditions>
+		<insert>
+			dependencies {
+			implementation 'com.android.billingclient:billing:2.0.3'
+			}
+		</insert>
+	</buildGradleAdditions>
+	
 	<!-- optional additions to GameActivity onCreate in GameActivity.java -->
 	<gameActivityOnCreateAdditions>
 		<insert>
diff --git a/Engine/Plugins/Online/Android/OnlineSubsystemGooglePlay/Source/Private/OnlinePurchaseGooglePlay.cpp b/Engine/Plugins/Online/Android/OnlineSubsystemGooglePlay/Source/Private/OnlinePurchaseGooglePlay.cpp
--- a/Engine/Plugins/Online/Android/OnlineSubsystemGooglePlay/Source/Private/OnlinePurchaseGooglePlay.cpp
+++ b/Engine/Plugins/Online/Android/OnlineSubsystemGooglePlay/Source/Private/OnlinePurchaseGooglePlay.cpp
@@ -5,6 +5,8 @@
 #include "OnlineStoreGooglePlay.h"
 #include "OnlineError.h"
 #include "Misc/Base64.h"
+#include <jni.h>
+#include "Android/AndroidJavaEnv.h"
 
 #define LOCTEXT_NAMESPACE "OnlineSubsystemGooglePlay"
 #define GOOGLEPLAYUSER TEXT("GooglePlayUser")
@@ -199,6 +201,7 @@ void FOnlinePurchaseGooglePlay::GetReceipts(const FUniqueNetId& UserId, TArray<F
 		const TArray<TSharedRef<FPurchaseReceipt>>& UserCompletedTransactions = *UserCompletedTransactionsPtr;
 		for (int32 Idx = 0; Idx < UserCompletedTransactions.Num(); Idx++)
 		{
+			UE_LOG_ONLINE_PURCHASE(Verbose, TEXT("FOnlinePurchaseGooglePlay::GetReceipts - Adding UserCompletedTransaction to OutReceipts"));
 			const TSharedRef<FPurchaseReceipt>& Transaction = UserCompletedTransactions[Idx];
 			OutReceipts.Add(*Transaction);
 		}
@@ -207,14 +210,17 @@ void FOnlinePurchaseGooglePlay::GetReceipts(const FUniqueNetId& UserId, TArray<F
 	// Add purchases completed while "offline"
 	for (int32 Idx = 0; Idx < OfflineTransactions.Num(); Idx++)
 	{
+		UE_LOG_ONLINE_PURCHASE(Verbose, TEXT("FOnlinePurchaseGooglePlay::GetReceipts - Adding OfflineTransaction to OutReceipts"));
 		OutReceipts.Add(*OfflineTransactions[Idx]);
 	}
+
+	UE_LOG_ONLINE_PURCHASE(Log, TEXT("FOnlinePurchaseGooglePlay::GetReceipts - Final Number of outreciepts OutReceipts - %d"), OutReceipts.Num());
 }
 
 void FOnlinePurchaseGooglePlay::OnTransactionCompleteResponse(EGooglePlayBillingResponseCode InResponseCode, const FGoogleTransactionData& InTransactionData)
 {
-	UE_LOG_ONLINE_PURCHASE(Verbose, TEXT("FOnlinePurchaseGooglePlay::OnTransactionCompleteResponse %s"), ToString(InResponseCode));
-	UE_LOG_ONLINE_PURCHASE(Verbose, TEXT("3... Transaction: %s"), *InTransactionData.ToDebugString());
+	UE_LOG_ONLINE_PURCHASE(Log, TEXT("FOnlinePurchaseGooglePlay::OnTransactionCompleteResponse %s"), ToString(InResponseCode));
+	UE_LOG_ONLINE_PURCHASE(Log, TEXT("3... Transaction: %s"), *InTransactionData.ToDebugString());
 	EPurchaseTransactionState Result = ConvertGPResponseCodeToPurchaseTransactionState(InResponseCode);
 
 	FString UserIdStr = GOOGLEPLAYUSER;
@@ -224,6 +230,12 @@ void FOnlinePurchaseGooglePlay::OnTransactionCompleteResponse(EGooglePlayBilling
 		const TSharedRef<FOnlinePurchasePendingTransactionGooglePlay> UserPendingTransaction = *UserPendingTransactionPtr;
 		const FString& ErrorStr = InTransactionData.GetErrorStr();
 
+		if (Result == EPurchaseTransactionState::Canceled && !InTransactionData.GetOfferId().IsEmpty())
+		{
+			// When result is cancelled, but there is a sku in the transaction data, this is a deferred transaction
+			Result = EPurchaseTransactionState::Deferred;
+		}
+
 		if (Result == EPurchaseTransactionState::Purchased || Result == EPurchaseTransactionState::Restored)
 		{
 			if (!UserPendingTransaction->AddCompletedOffer(Result, InTransactionData))
@@ -250,6 +262,10 @@ void FOnlinePurchaseGooglePlay::OnTransactionCompleteResponse(EGooglePlayBilling
 			case EPurchaseTransactionState::Purchased:
 				FinalResult.bSucceeded = true;
 				break;
+			case EPurchaseTransactionState::Deferred:
+				FinalResult.SetFromErrorCode(TEXT("com.epicgames.purchase.deferred"));
+				FinalResult.ErrorMessage = !ErrorStr.IsEmpty() ? FText::FromString(ErrorStr) : LOCTEXT("GooglePlayTransactionDeferred", "Transaction Deferred");
+				break;
 			case EPurchaseTransactionState::Invalid:
 				FinalResult.SetFromErrorCode(TEXT("com.epicgames.purchase.invalid"));
 				FinalResult.ErrorMessage = !ErrorStr.IsEmpty() ? FText::FromString(ErrorStr) : LOCTEXT("GooglePlayInvalidState", "Invalid purchase result");
@@ -268,19 +284,42 @@ void FOnlinePurchaseGooglePlay::OnTransactionCompleteResponse(EGooglePlayBilling
 		TArray< TSharedRef<FPurchaseReceipt> >& UserCompletedTransactions = CompletedTransactions.FindOrAdd(UserIdStr);
 
 		PendingTransactions.Remove(UserIdStr);
-		UserCompletedTransactions.Add(FinalReceipt);
+
+		// If this is a deferred transaction, we will process it as an "offline" transaction, so don't complete
+		if (Result != EPurchaseTransactionState::Deferred)
+		{
+			UserCompletedTransactions.Add(FinalReceipt);
+		}
 
 		UserPendingTransaction->CheckoutCompleteDelegate.ExecuteIfBound(FinalResult, FinalReceipt);
 	}
 	else
 	{
-		UE_LOG_ONLINE_PURCHASE(Log, TEXT("No pending transaction found associated with this purchase completion event"));
+		// Need to populate offlineTransactions here. 
+		// Transactions that come in during login or other non explicit purchase moments are added to a receipts list for later redemption
+		UE_LOG_ONLINE_PURCHASE(Log, TEXT("Pending transaction completed offline"));
+		if (Result == EPurchaseTransactionState::Restored || Result == EPurchaseTransactionState::Purchased)
+		{
+			TSharedRef<FPurchaseReceipt> OfflineReceipt = FOnlinePurchasePendingTransactionGooglePlay::GenerateReceipt(InTransactionData);
+			OfflineTransactions.Add(OfflineReceipt);
+
+			// Queue this user to be updated about this next-tick on the game-thread, if they're not mid-purchase
+			TWeakPtr<FOnlinePurchaseGooglePlay, ESPMode::ThreadSafe> WeakThis(AsShared());
+			Subsystem->ExecuteNextTick([WeakThis]()
+			{
+				FOnlinePurchaseGooglePlayPtr StrongThis = WeakThis.Pin();
+				if (StrongThis.IsValid())
+				{
+					StrongThis->TriggerOnUnexpectedPurchaseReceiptDelegates(*FUniqueNetIdGooglePlay::EmptyId());
+				}
+			});
+		}
 	}
 }
 
 void FOnlinePurchaseGooglePlay::OnQueryExistingPurchasesComplete(EGooglePlayBillingResponseCode InResponseCode, const TArray<FGoogleTransactionData>& InExistingPurchases)
 {
-	UE_LOG_ONLINE_PURCHASE(Verbose, TEXT("FOnlinePurchaseGooglePlay::OnQueryExistingPurchasesComplete Response: %s Num: %d"), ToString(InResponseCode), InExistingPurchases.Num());
+	UE_LOG_ONLINE_PURCHASE(Log, TEXT("FOnlinePurchaseGooglePlay::OnQueryExistingPurchasesComplete Response: %s Num: %d"), ToString(InResponseCode), InExistingPurchases.Num());
 
 	if (bQueryingReceipts)
 	{
@@ -290,14 +329,14 @@ void FOnlinePurchaseGooglePlay::OnQueryExistingPurchasesComplete(EGooglePlayBill
 			EPurchaseTransactionState Result = ConvertGPResponseCodeToPurchaseTransactionState(InResponseCode);
 			for (const FGoogleTransactionData& Purchase : InExistingPurchases)
 			{
-				UE_LOG_ONLINE_PURCHASE(Verbose, TEXT("Adding existing receipt %s"), *Purchase.ToDebugString());
+				UE_LOG_ONLINE_PURCHASE(Log, TEXT("Adding existing receipt %s"), *Purchase.ToDebugString());
 				TSharedRef<FPurchaseReceipt> OfflineReceipt = FOnlinePurchasePendingTransactionGooglePlay::GenerateReceipt(Purchase);
 				OfflineTransactions.Add(OfflineReceipt);
 			}
 		}
 		else
 		{
-			UE_LOG_ONLINE_PURCHASE(Verbose, TEXT("OnQueryExistingPurchasesComplete failed"));
+			UE_LOG_ONLINE_PURCHASE(Log, TEXT("OnQueryExistingPurchasesComplete failed"));
 		}
 
 		Subsystem->ExecuteNextTick([this, bSuccess]()
@@ -395,4 +434,148 @@ bool FOnlinePurchasePendingTransactionGooglePlay::AddCompletedOffer(EPurchaseTra
 	return false;
 }
 
+JNI_METHOD void Java_com_epicgames_ue4_GooglePlayStoreHelper_nativeQueryExistingPurchasesComplete(JNIEnv* jenv, jobject thiz, jsize responseCode, jobjectArray ProductIDs, jobjectArray ProductTokens, jobjectArray ReceiptsData, jobjectArray Signatures)
+{
+	TArray<FGoogleTransactionData> ExistingPurchaseInfo;
+
+	EGooglePlayBillingResponseCode EGPResponse = (EGooglePlayBillingResponseCode)responseCode;
+
+	bool bWasSuccessful = (EGPResponse == EGooglePlayBillingResponseCode::Ok);
+	if (jenv && bWasSuccessful)
+	{
+		jsize NumProducts = jenv->GetArrayLength(ProductIDs);
+		jsize NumProductTokens = jenv->GetArrayLength(ProductTokens);
+		jsize NumReceipts = jenv->GetArrayLength(ReceiptsData);
+		jsize NumSignatures = jenv->GetArrayLength(Signatures);
+
+		ensure((NumProducts == NumProductTokens) && (NumProducts == NumReceipts) && (NumProducts == NumSignatures));
+
+		for (jsize Idx = 0; Idx < NumProducts; Idx++)
+		{
+			// Build the product information strings.
+			const auto OfferId = FJavaHelper::FStringFromLocalRef(jenv, (jstring)jenv->GetObjectArrayElement(ProductIDs, Idx));
+			const auto ProductToken = FJavaHelper::FStringFromLocalRef(jenv, (jstring)jenv->GetObjectArrayElement(ProductTokens, Idx));
+			const auto ReceiptData = FJavaHelper::FStringFromLocalRef(jenv, (jstring)jenv->GetObjectArrayElement(ReceiptsData, Idx));
+			const auto SignatureData = FJavaHelper::FStringFromLocalRef(jenv, (jstring)jenv->GetObjectArrayElement(Signatures, Idx));
+
+			FGoogleTransactionData ExistingPurchase(OfferId, ProductToken, ReceiptData, SignatureData);
+			ExistingPurchaseInfo.Add(ExistingPurchase);
+
+			FPlatformMisc::LowLevelOutputDebugStringf(TEXT("\Existing Product Identifier: %s"), *ExistingPurchase.ToDebugString());
+		}
+	}
+
+	DECLARE_CYCLE_STAT(TEXT("FSimpleDelegateGraphTask.QueryExistingPurchases"), STAT_FSimpleDelegateGraphTask_QueryExistingPurchases, STATGROUP_TaskGraphTasks);
+
+	FSimpleDelegateGraphTask::CreateAndDispatchWhenReady(
+		FSimpleDelegateGraphTask::FDelegate::CreateLambda([=]()
+	{
+		FPlatformMisc::LowLevelOutputDebugStringf(TEXT("Query existing purchases was completed %s\n"), bWasSuccessful ? TEXT("successfully") : TEXT("unsuccessfully"));
+		if (IOnlineSubsystem* const OnlineSub = IOnlineSubsystem::Get(GOOGLEPLAY_SUBSYSTEM))
+		{
+			FOnlineSubsystemGooglePlay* const OnlineSubGP = static_cast<FOnlineSubsystemGooglePlay* const>(OnlineSub);
+			if (OnlineSubGP)
+			{
+				OnlineSubGP->TriggerOnGooglePlayQueryExistingPurchasesCompleteDelegates(EGPResponse, ExistingPurchaseInfo);
+			}
+		}
+	}),
+		GET_STATID(STAT_FSimpleDelegateGraphTask_QueryExistingPurchases),
+		nullptr,
+		ENamedThreads::GameThread
+		);
+}
+
+JNI_METHOD void Java_com_epicgames_ue4_GooglePlayStoreHelper_nativeRestorePurchasesComplete(JNIEnv* jenv, jobject thiz, jsize responseCode, jobjectArray ProductIDs, jobjectArray ProductTokens, jobjectArray ReceiptsData, jobjectArray Signatures)
+{
+	TArray<FGoogleTransactionData> RestoredPurchaseInfo;
+
+	EGooglePlayBillingResponseCode EGPResponse = (EGooglePlayBillingResponseCode)responseCode;
+	bool bWasSuccessful = (EGPResponse == EGooglePlayBillingResponseCode::Ok);
+
+	if (jenv && bWasSuccessful)
+	{
+		jsize NumProducts = jenv->GetArrayLength(ProductIDs);
+		jsize NumProductTokens = jenv->GetArrayLength(ProductTokens);
+		jsize NumReceipts = jenv->GetArrayLength(ReceiptsData);
+		jsize NumSignatures = jenv->GetArrayLength(Signatures);
+
+		ensure((NumProducts == NumProductTokens) && (NumProducts == NumReceipts) && (NumProducts == NumSignatures));
+
+		for (jsize Idx = 0; Idx < NumProducts; Idx++)
+		{
+			const auto OfferId = FJavaHelper::FStringFromLocalRef(jenv, (jstring)jenv->GetObjectArrayElement(ProductIDs, Idx));
+			const auto ProductToken = FJavaHelper::FStringFromLocalRef(jenv, (jstring)jenv->GetObjectArrayElement(ProductTokens, Idx));
+			const auto ReceiptData = FJavaHelper::FStringFromLocalRef(jenv, (jstring)jenv->GetObjectArrayElement(ReceiptsData, Idx));
+			const auto SignatureData = FJavaHelper::FStringFromLocalRef(jenv, (jstring)jenv->GetObjectArrayElement(Signatures, Idx));
+
+			FGoogleTransactionData RestoredPurchase(OfferId, ProductToken, ReceiptData, SignatureData);
+			RestoredPurchaseInfo.Add(RestoredPurchase);
+
+			FPlatformMisc::LowLevelOutputDebugStringf(TEXT("Restored Transaction: %s"), *RestoredPurchase.ToDebugString());
+		}
+	}
+
+	DECLARE_CYCLE_STAT(TEXT("FSimpleDelegateGraphTask.RestorePurchases"), STAT_FSimpleDelegateGraphTask_RestorePurchases, STATGROUP_TaskGraphTasks);
+
+	FSimpleDelegateGraphTask::CreateAndDispatchWhenReady(
+		FSimpleDelegateGraphTask::FDelegate::CreateLambda([=]()
+	{
+		FPlatformMisc::LowLevelOutputDebugStringf(TEXT("Restoring In-App Purchases was completed  %s\n"), bWasSuccessful ? TEXT("successfully") : TEXT("unsuccessfully"));
+
+		if (IOnlineSubsystem* const OnlineSub = IOnlineSubsystem::Get(GOOGLEPLAY_SUBSYSTEM))
+		{
+			FOnlineSubsystemGooglePlay* const OnlineSubGP = static_cast<FOnlineSubsystemGooglePlay* const>(OnlineSub);
+			if (OnlineSubGP)
+			{
+				OnlineSubGP->TriggerOnGooglePlayRestorePurchasesCompleteDelegates(EGPResponse, RestoredPurchaseInfo);
+			}
+		}
+	}),
+		GET_STATID(STAT_FSimpleDelegateGraphTask_RestorePurchases),
+		nullptr,
+		ENamedThreads::GameThread
+		);
+}
+
+JNI_METHOD void Java_com_epicgames_ue4_GooglePlayStoreHelper_nativePurchaseComplete(JNIEnv* jenv, jobject thiz, jsize responseCode, jstring productId, jstring productToken, jstring receiptData, jstring signature)
+{
+	FString ProductId, ProductToken, ReceiptData, Signature;
+	EGooglePlayBillingResponseCode EGPResponse = (EGooglePlayBillingResponseCode)responseCode;
+
+	bool bWasSuccessful = (EGPResponse == EGooglePlayBillingResponseCode::Ok);
+
+	//Store off results, because we will use them to determine if this is a deferred transaction.
+	ProductId = FJavaHelper::FStringFromParam(jenv, productId);
+	ProductToken = FJavaHelper::FStringFromParam(jenv, productToken);
+	ReceiptData = FJavaHelper::FStringFromParam(jenv, receiptData);
+	Signature = FJavaHelper::FStringFromParam(jenv, signature);
+
+
+	FGoogleTransactionData TransactionData(ProductId, ProductToken, ReceiptData, Signature);
+
+	FPlatformMisc::LowLevelOutputDebugStringf(TEXT("1... Response: %s, Transaction %s"), ToString(EGPResponse), *TransactionData.ToDebugString());
+
+	DECLARE_CYCLE_STAT(TEXT("FSimpleDelegateGraphTask.ProcessIapResult"), STAT_FSimpleDelegateGraphTask_ProcessIapResult, STATGROUP_TaskGraphTasks);
+
+	FSimpleDelegateGraphTask::CreateAndDispatchWhenReady(
+		FSimpleDelegateGraphTask::FDelegate::CreateLambda([=]()
+	{
+		FPlatformMisc::LowLevelOutputDebugStringf(TEXT("In-App Purchase was completed  %s\n"), bWasSuccessful ? TEXT("successfully") : TEXT("unsuccessfully"));
+		if (IOnlineSubsystem* const OnlineSub = IOnlineSubsystem::Get(GOOGLEPLAY_SUBSYSTEM))
+		{
+			FOnlineSubsystemGooglePlay* const OnlineSubGP = static_cast<FOnlineSubsystemGooglePlay* const>(OnlineSub);
+			if (OnlineSubGP)
+			{
+				FPlatformMisc::LowLevelOutputDebugStringf(TEXT("2... Response %s Transaction %s"), ToString(EGPResponse), *TransactionData.ToDebugString());
+				OnlineSubGP->TriggerOnGooglePlayProcessPurchaseCompleteDelegates(EGPResponse, TransactionData);
+			}
+		}
+	}),
+		GET_STATID(STAT_FSimpleDelegateGraphTask_ProcessIapResult),
+		nullptr,
+		ENamedThreads::GameThread
+		);
+}
+
 #undef LOCTEXT_NAMESPACE
diff --git a/Engine/Plugins/Online/Android/OnlineSubsystemGooglePlay/Source/Private/OnlineStoreGooglePlay.cpp b/Engine/Plugins/Online/Android/OnlineSubsystemGooglePlay/Source/Private/OnlineStoreGooglePlay.cpp
--- a/Engine/Plugins/Online/Android/OnlineSubsystemGooglePlay/Source/Private/OnlineStoreGooglePlay.cpp
+++ b/Engine/Plugins/Online/Android/OnlineSubsystemGooglePlay/Source/Private/OnlineStoreGooglePlay.cpp
@@ -8,6 +8,8 @@
 #include "Internationalization/Culture.h"
 #include "Internationalization/FastDecimalFormat.h"
 #include "Misc/ConfigCacheIni.h"
+#include <jni.h>
+#include "Android/AndroidJavaEnv.h"
 
 FOnlineStoreGooglePlayV2::FOnlineStoreGooglePlayV2(FOnlineSubsystemGooglePlay* InSubsystem)
 	: bIsQueryInFlight(false)
@@ -48,50 +50,12 @@ void FOnlineStoreGooglePlayV2::Init()
 	AndroidThunkCpp_Iap_SetupIapService(GooglePlayLicenseKey);
 }
 
-TSharedRef<FOnlineStoreOffer> ConvertProductToStoreOffer(const FInAppPurchaseProductInfo& Product)
-{
-	TSharedRef<FOnlineStoreOffer> NewProductInfo = MakeShareable(new FOnlineStoreOffer());
-
-	NewProductInfo->OfferId = Product.Identifier;
-
-	FString Title = Product.DisplayName;
-	int32 OpenParenIdx = -1;
-	int32 CloseParenIdx = -1;
-	if (Title.FindLastChar(TEXT(')'), CloseParenIdx) && Title.FindLastChar(TEXT('('), OpenParenIdx) && (OpenParenIdx < CloseParenIdx))
-	{
-		Title.LeftInline(OpenParenIdx);
-		Title.TrimEndInline();
-	}
-
-	NewProductInfo->Title = FText::FromString(Title);
-	NewProductInfo->Description = FText::FromString(Product.DisplayDescription); // Google has only one description, map it to (short) description to match iOS
-	//NewProductInfo->LongDescription = FText::FromString(Product.DisplayDescription); // leave this empty so we know it's not set (client can apply more info from MCP)
-	NewProductInfo->PriceText = FText::FromString(Product.DisplayPrice);
-	NewProductInfo->CurrencyCode = Product.CurrencyCode;
-
-	// Convert the backend stated price into its base units
-	FInternationalization& I18N = FInternationalization::Get();
-	const FCulture& Culture = *I18N.GetCurrentCulture();
-
-	const FDecimalNumberFormattingRules& FormattingRules = Culture.GetCurrencyFormattingRules(NewProductInfo->CurrencyCode);
-	const FNumberFormattingOptions& FormattingOptions = FormattingRules.CultureDefaultFormattingOptions;
-	double Val = static_cast<double>(Product.RawPrice) * static_cast<double>(FMath::Pow(10.0f, FormattingOptions.MaximumFractionalDigits));
-
-	NewProductInfo->NumericPrice = FMath::TruncToInt(Val + 0.5);
-
-	// Google doesn't support these fields, set to min and max defaults
-	NewProductInfo->ReleaseDate = FDateTime::MinValue();
-	NewProductInfo->ExpirationDate = FDateTime::MaxValue();
-
-	return NewProductInfo;
-}
-
 void FOnlineStoreGooglePlayV2::OnGooglePlayAvailableIAPQueryComplete(EGooglePlayBillingResponseCode InResponseCode, const TArray<FInAppPurchaseProductInfo>& InProvidedProductInformation)
 { 
 	TArray<FUniqueOfferId> OfferIds;
-	for (const FInAppPurchaseProductInfo& Product : InProvidedProductInformation)
+	for (const FOnlineStoreOffer& Product : InProvidedProductInformation)
 	{
-		TSharedRef<FOnlineStoreOffer> NewProductOffer = ConvertProductToStoreOffer(Product);
+		TSharedRef<FOnlineStoreOffer> NewProductOffer = MakeShared<FOnlineStoreOffer>(Product);
 
 		AddOffer(NewProductOffer);
 		OfferIds.Add(NewProductOffer->OfferId);
@@ -194,3 +158,104 @@ TSharedPtr<FOnlineStoreOffer> FOnlineStoreGooglePlayV2::GetOffer(const FUniqueOf
 
 	return Result;
 }
+
+JNI_METHOD void Java_com_epicgames_ue4_GooglePlayStoreHelper_nativeQueryComplete(JNIEnv* jenv, jobject thiz, jsize responseCode, jobjectArray productIDs, jobjectArray titles, jobjectArray descriptions, jobjectArray prices, jfloatArray pricesRaw, jobjectArray currencyCodes, jobjectArray originalJson)
+{
+	TArray<FOnlineStoreOffer> ProvidedProductInformation;
+	EGooglePlayBillingResponseCode EGPResponse = (EGooglePlayBillingResponseCode)responseCode;
+	bool bWasSuccessful = (EGPResponse == EGooglePlayBillingResponseCode::Ok);
+
+	if (jenv && bWasSuccessful)
+	{
+		jsize NumProducts = jenv->GetArrayLength(productIDs);
+		jsize NumTitles = jenv->GetArrayLength(titles);
+		jsize NumDescriptions = jenv->GetArrayLength(descriptions);
+		jsize NumPrices = jenv->GetArrayLength(prices);
+		jsize NumPricesRaw = jenv->GetArrayLength(pricesRaw);
+		jsize NumCurrencyCodes = jenv->GetArrayLength(currencyCodes);
+		jsize NumJsonStrings = jenv->GetArrayLength(originalJson);
+
+		ensure((NumProducts == NumTitles) && (NumProducts == NumDescriptions) && (NumProducts == NumPrices) && (NumProducts == NumPricesRaw) && (NumProducts == NumCurrencyCodes) && (NumProducts == NumJsonStrings));
+
+		jfloat *PricesRaw = jenv->GetFloatArrayElements(pricesRaw, 0);
+
+		for (jsize Idx = 0; Idx < NumProducts; Idx++)
+		{
+			// Build the product information strings.
+
+			FOnlineStoreOffer NewProductInfo;
+
+			NewProductInfo.OfferId = FJavaHelper::FStringFromLocalRef(jenv, (jstring)jenv->GetObjectArrayElement(productIDs, Idx));
+			
+			int32 OpenParenIdx = -1;
+			int32 CloseParenIdx = -1;
+			FString Title = FJavaHelper::FStringFromLocalRef(jenv, (jstring)jenv->GetObjectArrayElement(titles, Idx));
+			if (Title.FindLastChar(TEXT(')'), CloseParenIdx) && Title.FindLastChar(TEXT('('), OpenParenIdx) && (OpenParenIdx < CloseParenIdx))
+			{
+				Title = Title.Left(OpenParenIdx).TrimEnd();
+			}
+			NewProductInfo.Title = FText::FromString(Title);
+
+			NewProductInfo.Description = FText::FromString(FJavaHelper::FStringFromLocalRef(jenv, (jstring)jenv->GetObjectArrayElement(descriptions, Idx)));
+			NewProductInfo.PriceText = FText::FromString(FJavaHelper::FStringFromLocalRef(jenv, (jstring)jenv->GetObjectArrayElement(prices, Idx)));
+
+			// Convert the backend stated price into its base units
+			FInternationalization& I18N = FInternationalization::Get();
+			const FCulture& Culture = *I18N.GetCurrentCulture();
+
+			const FDecimalNumberFormattingRules& FormattingRules = Culture.GetCurrencyFormattingRules(NewProductInfo.CurrencyCode);
+			const FNumberFormattingOptions& FormattingOptions = FormattingRules.CultureDefaultFormattingOptions;
+			double Val = static_cast<double>(PricesRaw[Idx]) * static_cast<double>(FMath::Pow(10.0f, FormattingOptions.MaximumFractionalDigits));
+
+			NewProductInfo.NumericPrice = FMath::TruncToInt(Val + 0.5);
+			NewProductInfo.CurrencyCode = FJavaHelper::FStringFromLocalRef(jenv, (jstring)jenv->GetObjectArrayElement(currencyCodes, Idx));
+
+			//Loop through original json data and populate dynamic map.
+			TSharedPtr<FJsonObject> JsonObject;
+			TSharedRef<TJsonReader<>> JsonReader = TJsonReaderFactory<>::Create(FJavaHelper::FStringFromLocalRef(jenv, (jstring)jenv->GetObjectArrayElement(originalJson, Idx)));
+			if (FJsonSerializer::Deserialize(JsonReader, JsonObject) && JsonObject.IsValid())
+			{
+				for (TPair<FString, TSharedPtr<FJsonValue>> JsonValue : JsonObject->Values)
+				{
+					NewProductInfo.DynamicFields.Add(JsonValue.Key, JsonValue.Value->AsString());
+				}
+			}
+
+			NewProductInfo.ReleaseDate = FDateTime::MinValue();
+			NewProductInfo.ExpirationDate = FDateTime::MaxValue();
+
+			ProvidedProductInformation.Add(NewProductInfo);
+
+			FPlatformMisc::LowLevelOutputDebugStringf(TEXT("\nProduct Identifier: %s, Name: %s, Description: %s, Price: %s, Price Raw: %d, Currency Code: %s\n"),
+				*NewProductInfo.OfferId,
+				*NewProductInfo.Title.ToString(),
+				*NewProductInfo.Description.ToString(),
+				*NewProductInfo.GetDisplayPrice().ToString(),
+				NewProductInfo.NumericPrice,
+				*NewProductInfo.CurrencyCode);
+		}
+	}
+
+	DECLARE_CYCLE_STAT(TEXT("FSimpleDelegateGraphTask.ProcessQueryIapResult"), STAT_FSimpleDelegateGraphTask_ProcessQueryIapResult, STATGROUP_TaskGraphTasks);
+
+	FPlatformMisc::LowLevelOutputDebugStringf(TEXT("Adding task Success: %d Response: %s"), bWasSuccessful, ToString(EGPResponse));
+
+	FSimpleDelegateGraphTask::CreateAndDispatchWhenReady(
+		FSimpleDelegateGraphTask::FDelegate::CreateLambda([=]()
+	{
+		if (IOnlineSubsystem* const OnlineSub = IOnlineSubsystem::Get(GOOGLEPLAY_SUBSYSTEM))
+		{
+			FOnlineSubsystemGooglePlay* const OnlineSubGP = static_cast<FOnlineSubsystemGooglePlay* const>(OnlineSub);
+			if (OnlineSubGP)
+			{
+				FPlatformMisc::LowLevelOutputDebugStringf(TEXT("TriggerOnGooglePlayAvailableIAPQueryCompleteDelegates %s Size: %d"), ToString(EGPResponse), ProvidedProductInformation.Num());
+				OnlineSubGP->TriggerOnGooglePlayAvailableIAPQueryCompleteDelegates(EGPResponse, ProvidedProductInformation);
+			}
+		}
+		FPlatformMisc::LowLevelOutputDebugStringf(TEXT("In-App Purchase query was completed  %s\n"), bWasSuccessful ? TEXT("successfully") : TEXT("unsuccessfully"));
+	}),
+		GET_STATID(STAT_FSimpleDelegateGraphTask_ProcessQueryIapResult),
+		nullptr,
+		ENamedThreads::GameThread
+		);
+}
\ No newline at end of file
diff --git a/Engine/Plugins/Online/Android/OnlineSubsystemGooglePlay/Source/Private/OnlineStoreGooglePlayCommon.h b/Engine/Plugins/Online/Android/OnlineSubsystemGooglePlay/Source/Private/OnlineStoreGooglePlayCommon.h
--- a/Engine/Plugins/Online/Android/OnlineSubsystemGooglePlay/Source/Private/OnlineStoreGooglePlayCommon.h
+++ b/Engine/Plugins/Online/Android/OnlineSubsystemGooglePlay/Source/Private/OnlineStoreGooglePlayCommon.h
@@ -1,7 +1,6 @@
 // Copyright Epic Games, Inc. All Rights Reserved.
 
 #pragma once
-
 #include "Interfaces/OnlineStoreInterface.h"
 #include "Interfaces/OnlineStoreInterfaceV2.h"
 #include "Interfaces/OnlinePurchaseInterface.h"
@@ -105,18 +104,6 @@ enum class EInAppPurchaseResult : uint8
 	Cancelled,
 };
 
-/**
- * Implementation of the Platform Purchase receipt. For this we provide an identifier and the encrypted data.
- */
-class FGooglePlayPurchaseReceipt : public IPlatformPurchaseReceipt
-{
-public:
-	// Product identifier
-	FString Identifier;
-
-	// The encrypted receipt data
-	FString Data;
-};
 
  /**
   * Delegate fired when an IAP query for available offers has completed
@@ -124,7 +111,7 @@ public:
   * @param Response response from GooglePlay backend
   * @param ProvidedProductInformation list of offers returned in response to a query on available offer ids
   */
-DECLARE_MULTICAST_DELEGATE_TwoParams(FOnGooglePlayAvailableIAPQueryComplete, EGooglePlayBillingResponseCode /*Response*/, const TArray<FInAppPurchaseProductInfo>& /*ProvidedProductInformation*/);
+DECLARE_MULTICAST_DELEGATE_TwoParams(FOnGooglePlayAvailableIAPQueryComplete, EGooglePlayBillingResponseCode /*Response*/, const TArray<FOnlineStoreOffer>& /*ProvidedProductInformation*/);
 typedef FOnGooglePlayAvailableIAPQueryComplete::FDelegate FOnGooglePlayAvailableIAPQueryCompleteDelegate;
 
 /**
diff --git a/Engine/Plugins/Online/Android/OnlineSubsystemGooglePlay/Source/Private/OnlineStoreInterfaceGooglePlay.cpp b/Engine/Plugins/Online/Android/OnlineSubsystemGooglePlay/Source/Private/OnlineStoreInterfaceGooglePlay.cpp
--- a/Engine/Plugins/Online/Android/OnlineSubsystemGooglePlay/Source/Private/OnlineStoreInterfaceGooglePlay.cpp
+++ b/Engine/Plugins/Online/Android/OnlineSubsystemGooglePlay/Source/Private/OnlineStoreInterfaceGooglePlay.cpp
@@ -9,7 +9,9 @@
 #include "OnlineSubsystemGooglePlay.h"
 #include <jni.h>
 #include "Android/AndroidJavaEnv.h"
-
+#include "Internationalization/Internationalization.h"
+#include "Internationalization/Culture.h"
+#include "Internationalization/FastDecimalFormat.h"
 
 ////////////////////////////////////////////////////////////////////
 /// FOnlineStoreGooglePlay implementation
@@ -79,86 +81,44 @@ bool FOnlineStoreGooglePlay::QueryForAvailablePurchases(const TArray<FString>& P
 	return true;
 }
 
-
-JNI_METHOD void Java_com_epicgames_ue4_GooglePlayStoreHelper_nativeQueryComplete(JNIEnv* jenv, jobject thiz, jsize responseCode, jobjectArray productIDs, jobjectArray titles, jobjectArray descriptions, jobjectArray prices, jfloatArray pricesRaw, jobjectArray currencyCodes)
+TSharedRef<FInAppPurchaseProductInfo> ConvertStoreOfferToProduct(const FOnlineStoreOffer& Offer)
 {
-	TArray<FInAppPurchaseProductInfo> ProvidedProductInformation;
-	EGooglePlayBillingResponseCode EGPResponse = (EGooglePlayBillingResponseCode)responseCode;
-	bool bWasSuccessful = (EGPResponse == EGooglePlayBillingResponseCode::Ok);
-
-	if (jenv && bWasSuccessful)
-	{
-		jsize NumProducts = jenv->GetArrayLength(productIDs);
-		jsize NumTitles = jenv->GetArrayLength(titles);
-		jsize NumDescriptions = jenv->GetArrayLength(descriptions);
-		jsize NumPrices = jenv->GetArrayLength(prices);
-		jsize NumPricesRaw = jenv->GetArrayLength(pricesRaw);
-		jsize NumCurrencyCodes = jenv->GetArrayLength(currencyCodes);
+	TSharedRef<FInAppPurchaseProductInfo> NewProductInfo = MakeShareable(new FInAppPurchaseProductInfo());
 
-		ensure((NumProducts == NumTitles) && (NumProducts == NumDescriptions) && (NumProducts == NumPrices) && (NumProducts == NumPricesRaw) && (NumProducts == NumCurrencyCodes));
+	NewProductInfo->Identifier = Offer.OfferId;
+	NewProductInfo->DisplayName = Offer.Title.ToString();
 
-		jfloat *PricesRaw = jenv->GetFloatArrayElements(pricesRaw, 0);
+	NewProductInfo->DisplayDescription = Offer.Description.ToString(); // Google has only one description, map it to (short) description to match iOS
+	NewProductInfo->DisplayPrice = Offer.PriceText.ToString();
+	NewProductInfo->CurrencyCode = Offer.CurrencyCode;
 
-		for (jsize Idx = 0; Idx < NumProducts; Idx++)
-		{
-			// Build the product information strings.
-
-			FInAppPurchaseProductInfo NewProductInfo;
-
-			NewProductInfo.Identifier = FJavaHelper::FStringFromLocalRef(jenv, (jstring)jenv->GetObjectArrayElement(productIDs, Idx));
-			NewProductInfo.DisplayName = FJavaHelper::FStringFromLocalRef(jenv, (jstring)jenv->GetObjectArrayElement(titles, Idx));
-			NewProductInfo.DisplayDescription = FJavaHelper::FStringFromLocalRef(jenv, (jstring)jenv->GetObjectArrayElement(descriptions, Idx));
-			NewProductInfo.DisplayPrice = FJavaHelper::FStringFromLocalRef(jenv, (jstring)jenv->GetObjectArrayElement(prices, Idx));
-			
-			NewProductInfo.RawPrice = PricesRaw[Idx];
-			
-			NewProductInfo.CurrencyCode = FJavaHelper::FStringFromLocalRef(jenv, (jstring)jenv->GetObjectArrayElement(currencyCodes, Idx));
-
-			ProvidedProductInformation.Add(NewProductInfo);
-
-			FPlatformMisc::LowLevelOutputDebugStringf(TEXT("\nProduct Identifier: %s, Name: %s, Description: %s, Price: %s, Price Raw: %f, Currency Code: %s\n"),
-				*NewProductInfo.Identifier,
-				*NewProductInfo.DisplayName,
-				*NewProductInfo.DisplayDescription,
-				*NewProductInfo.DisplayPrice,
-				NewProductInfo.RawPrice,
-				*NewProductInfo.CurrencyCode);
-		}
-	}
+	// Convert the backend stated price into its base units
+	FInternationalization& I18N = FInternationalization::Get();
+	const FCulture& Culture = *I18N.GetCurrentCulture();
 
-	DECLARE_CYCLE_STAT(TEXT("FSimpleDelegateGraphTask.ProcessQueryIapResult"), STAT_FSimpleDelegateGraphTask_ProcessQueryIapResult, STATGROUP_TaskGraphTasks);
+	const FDecimalNumberFormattingRules& FormattingRules = Culture.GetCurrencyFormattingRules(NewProductInfo->CurrencyCode);
+	const FNumberFormattingOptions& FormattingOptions = FormattingRules.CultureDefaultFormattingOptions;
+	double Val = static_cast<double>(Offer.NumericPrice) / static_cast<double>(FMath::Pow(10.0f, FormattingOptions.MaximumFractionalDigits));
 
-	FPlatformMisc::LowLevelOutputDebugStringf(TEXT("Adding task Success: %d Response: %s"), bWasSuccessful, ToString(EGPResponse));
+	NewProductInfo->RawPrice = Val;
 
-	FSimpleDelegateGraphTask::CreateAndDispatchWhenReady(
-		FSimpleDelegateGraphTask::FDelegate::CreateLambda([=]()
-	{
-		if (IOnlineSubsystem* const OnlineSub = IOnlineSubsystem::Get(GOOGLEPLAY_SUBSYSTEM))
-		{
-			FOnlineSubsystemGooglePlay* const OnlineSubGP = static_cast<FOnlineSubsystemGooglePlay* const>(OnlineSub);
-			if (OnlineSubGP)
-			{
-				FPlatformMisc::LowLevelOutputDebugStringf(TEXT("TriggerOnGooglePlayAvailableIAPQueryCompleteDelegates %s Size: %d"), ToString(EGPResponse), ProvidedProductInformation.Num());
-				OnlineSubGP->TriggerOnGooglePlayAvailableIAPQueryCompleteDelegates(EGPResponse, ProvidedProductInformation);
-			}
-		}
-		FPlatformMisc::LowLevelOutputDebugStringf(TEXT("In-App Purchase query was completed  %s\n"), bWasSuccessful ? TEXT("successfully") : TEXT("unsuccessfully"));
-	}),
-		GET_STATID(STAT_FSimpleDelegateGraphTask_ProcessQueryIapResult),
-		nullptr,
-		ENamedThreads::GameThread
-		);
+	return NewProductInfo;
 }
 
-
-void FOnlineStoreGooglePlay::OnGooglePlayAvailableIAPQueryComplete(EGooglePlayBillingResponseCode InResponseCode, const TArray<FInAppPurchaseProductInfo>& AvailablePurchases)
+void FOnlineStoreGooglePlay::OnGooglePlayAvailableIAPQueryComplete(EGooglePlayBillingResponseCode InResponseCode, const TArray<FOnlineStoreOffer>& AvailablePurchases)
 {
 	UE_LOG_ONLINE_STORE(Display, TEXT("FOnlineStoreGooglePlay::OnGooglePlayAvailableIAPQueryComplete"));
 
+	TArray<FInAppPurchaseProductInfo> TempInfos;
+	for (const FOnlineStoreOffer& AvailablePurchase : AvailablePurchases)
+	{
+		TempInfos.Add(*ConvertStoreOfferToProduct(AvailablePurchase));
+	}
+
 	if (ReadObject.IsValid())
 	{
 		ReadObject->ReadState = (InResponseCode == EGooglePlayBillingResponseCode::Ok) ? EOnlineAsyncTaskState::Done : EOnlineAsyncTaskState::Failed;
-		ReadObject->ProvidedProductInformation.Insert(AvailablePurchases, 0);
+		ReadObject->ProvidedProductInformation.Insert(TempInfos, 0);
 	}
 
 	CurrentQueryTask->ProcessQueryAvailablePurchasesResults(InResponseCode == EGooglePlayBillingResponseCode::Ok);
@@ -205,48 +165,6 @@ bool FOnlineStoreGooglePlay::BeginPurchase(const FInAppPurchaseProductRequest& P
 	return bCreatedNewTransaction;
 }
 
-
-JNI_METHOD void Java_com_epicgames_ue4_GooglePlayStoreHelper_nativePurchaseComplete(JNIEnv* jenv, jobject thiz, jsize responseCode, jstring productId, jstring productToken, jstring receiptData, jstring signature)
-{
-	FString ProductId, ProductToken, ReceiptData, Signature;
-	EGooglePlayBillingResponseCode EGPResponse = (EGooglePlayBillingResponseCode)responseCode;
-
-	bool bWasSuccessful = (EGPResponse == EGooglePlayBillingResponseCode::Ok);
-	if (bWasSuccessful)
-	{
-		ProductId = FJavaHelper::FStringFromParam(jenv, productId);
-		ProductToken = FJavaHelper::FStringFromParam(jenv, productToken);
-		ReceiptData = FJavaHelper::FStringFromParam(jenv, receiptData);
-		Signature = FJavaHelper::FStringFromParam(jenv, signature);
-	}
-
-	FGoogleTransactionData TransactionData(ProductId, ProductToken, ReceiptData, Signature);
-
-	FPlatformMisc::LowLevelOutputDebugStringf(TEXT("1... Response: %s, Transaction %s"), ToString(EGPResponse), *TransactionData.ToDebugString());
-
-	DECLARE_CYCLE_STAT(TEXT("FSimpleDelegateGraphTask.ProcessIapResult"), STAT_FSimpleDelegateGraphTask_ProcessIapResult, STATGROUP_TaskGraphTasks);
-
-	FSimpleDelegateGraphTask::CreateAndDispatchWhenReady(
-		FSimpleDelegateGraphTask::FDelegate::CreateLambda([=]()
-	{
-		FPlatformMisc::LowLevelOutputDebugStringf(TEXT("In-App Purchase was completed  %s\n"), bWasSuccessful ? TEXT("successfully") : TEXT("unsuccessfully"));
-		if (IOnlineSubsystem* const OnlineSub = IOnlineSubsystem::Get(GOOGLEPLAY_SUBSYSTEM))
-		{
-			FOnlineSubsystemGooglePlay* const OnlineSubGP = static_cast<FOnlineSubsystemGooglePlay* const>(OnlineSub);
-			if (OnlineSubGP)
-			{
-				FPlatformMisc::LowLevelOutputDebugStringf(TEXT("2... Response %s Transaction %s"), ToString(EGPResponse), *TransactionData.ToDebugString());
-				OnlineSubGP->TriggerOnGooglePlayProcessPurchaseCompleteDelegates(EGPResponse, TransactionData);
-			}
-		}
-	}),
-		GET_STATID(STAT_FSimpleDelegateGraphTask_ProcessIapResult),
-		nullptr,
-		ENamedThreads::GameThread
-		);
-}
-
-
 void FOnlineStoreGooglePlay::OnProcessPurchaseResult(EGooglePlayBillingResponseCode InResponseCode, const FGoogleTransactionData& InTransactionData)
 {
 	UE_LOG_ONLINE_STORE(Display, TEXT("FOnlineStoreGooglePlay::OnProcessPurchaseResult"));
@@ -308,61 +226,6 @@ bool FOnlineStoreGooglePlay::RestorePurchases(const TArray<FInAppPurchaseProduct
 	return bSentAQueryRequest;
 }
 
-JNI_METHOD void Java_com_epicgames_ue4_GooglePlayStoreHelper_nativeRestorePurchasesComplete(JNIEnv* jenv, jobject thiz, jsize responseCode, jobjectArray ProductIDs, jobjectArray ProductTokens, jobjectArray ReceiptsData, jobjectArray Signatures)
-{
-	TArray<FGoogleTransactionData> RestoredPurchaseInfo;
-
-	EGooglePlayBillingResponseCode EGPResponse = (EGooglePlayBillingResponseCode)responseCode;
-	bool bWasSuccessful = (EGPResponse == EGooglePlayBillingResponseCode::Ok);
-
-	if (jenv && bWasSuccessful)
-	{
-		jsize NumProducts = jenv->GetArrayLength(ProductIDs);
-		jsize NumProductTokens = jenv->GetArrayLength(ProductTokens);
-		jsize NumReceipts = jenv->GetArrayLength(ReceiptsData);
-		jsize NumSignatures = jenv->GetArrayLength(Signatures);
-
-		ensure((NumProducts == NumProductTokens) && (NumProducts == NumReceipts) && (NumProducts == NumSignatures));
-
-		for (jsize Idx = 0; Idx < NumProducts; Idx++)
-		{
-			// Build the restore product information strings.
-			FInAppPurchaseRestoreInfo RestoreInfo;
-
-			const auto OfferId = FJavaHelper::FStringFromLocalRef(jenv, (jstring)jenv->GetObjectArrayElement(ProductIDs, Idx));
-			const auto ProductToken = FJavaHelper::FStringFromLocalRef(jenv, (jstring)jenv->GetObjectArrayElement(ProductTokens, Idx));
-			const auto ReceiptData = FJavaHelper::FStringFromLocalRef(jenv, (jstring)jenv->GetObjectArrayElement(ReceiptsData, Idx));
-			const auto SignatureData = FJavaHelper::FStringFromLocalRef(jenv, (jstring)jenv->GetObjectArrayElement(Signatures, Idx));
-			
-			FGoogleTransactionData RestoredPurchase(OfferId, ProductToken, ReceiptData, SignatureData);
-			RestoredPurchaseInfo.Add(RestoredPurchase);
-
-			FPlatformMisc::LowLevelOutputDebugStringf(TEXT("Restored Transaction: %s"), *RestoredPurchase.ToDebugString());
-		}
-	}
-
-	DECLARE_CYCLE_STAT(TEXT("FSimpleDelegateGraphTask.RestorePurchases"), STAT_FSimpleDelegateGraphTask_RestorePurchases, STATGROUP_TaskGraphTasks);
-
-	FSimpleDelegateGraphTask::CreateAndDispatchWhenReady(
-		FSimpleDelegateGraphTask::FDelegate::CreateLambda([=]()
-	{
-		FPlatformMisc::LowLevelOutputDebugStringf(TEXT("Restoring In-App Purchases was completed  %s\n"), bWasSuccessful ? TEXT("successfully") : TEXT("unsuccessfully"));
-
-		if (IOnlineSubsystem* const OnlineSub = IOnlineSubsystem::Get(GOOGLEPLAY_SUBSYSTEM))
-		{
-			FOnlineSubsystemGooglePlay* const OnlineSubGP = static_cast<FOnlineSubsystemGooglePlay* const>(OnlineSub);
-			if (OnlineSubGP)
-			{
-				OnlineSubGP->TriggerOnGooglePlayRestorePurchasesCompleteDelegates(EGPResponse, RestoredPurchaseInfo);
-			}
-		}
-	}),
-		GET_STATID(STAT_FSimpleDelegateGraphTask_RestorePurchases),
-		nullptr,
-		ENamedThreads::GameThread
-		);
-}
-
 void FOnlineStoreGooglePlay::OnRestorePurchasesComplete(EGooglePlayBillingResponseCode InResponseCode, const TArray<FGoogleTransactionData>& InRestoredPurchases)
 {
 	UE_LOG_ONLINE_STORE(Verbose, TEXT("FOnlineStoreGooglePlay::OnRestorePurchasesComplete Response: %s Num: %d"), ToString(InResponseCode), InRestoredPurchases.Num());
@@ -390,54 +253,3 @@ void FOnlineStoreGooglePlay::OnRestorePurchasesComplete(EGooglePlayBillingRespon
 	TriggerOnInAppPurchaseRestoreCompleteDelegates(IAPState);
 }
 
-JNI_METHOD void Java_com_epicgames_ue4_GooglePlayStoreHelper_nativeQueryExistingPurchasesComplete(JNIEnv* jenv, jobject thiz, jsize responseCode, jobjectArray ProductIDs, jobjectArray ProductTokens, jobjectArray ReceiptsData, jobjectArray Signatures)
-{
-	TArray<FGoogleTransactionData> ExistingPurchaseInfo;
-
-	EGooglePlayBillingResponseCode EGPResponse = (EGooglePlayBillingResponseCode)responseCode;
-
-	bool bWasSuccessful = (EGPResponse == EGooglePlayBillingResponseCode::Ok);
-	if (jenv && bWasSuccessful)
-	{
-		jsize NumProducts = jenv->GetArrayLength(ProductIDs);
-		jsize NumProductTokens = jenv->GetArrayLength(ProductTokens);
-		jsize NumReceipts = jenv->GetArrayLength(ReceiptsData);
-		jsize NumSignatures = jenv->GetArrayLength(Signatures);
-
-		ensure((NumProducts == NumProductTokens) && (NumProducts == NumReceipts) && (NumProducts == NumSignatures));
-
-		for (jsize Idx = 0; Idx < NumProducts; Idx++)
-		{
-			// Build the product information strings.
-			const auto OfferId = FJavaHelper::FStringFromLocalRef(jenv, (jstring)jenv->GetObjectArrayElement(ProductIDs, Idx));
-			const auto ProductToken = FJavaHelper::FStringFromLocalRef(jenv, (jstring)jenv->GetObjectArrayElement(ProductTokens, Idx));
-			const auto ReceiptData = FJavaHelper::FStringFromLocalRef(jenv, (jstring)jenv->GetObjectArrayElement(ReceiptsData, Idx));
-			const auto SignatureData = FJavaHelper::FStringFromLocalRef(jenv, (jstring)jenv->GetObjectArrayElement(Signatures, Idx));
-			
-			FGoogleTransactionData ExistingPurchase(OfferId, ProductToken, ReceiptData, SignatureData);
-			ExistingPurchaseInfo.Add(ExistingPurchase);
-
-			FPlatformMisc::LowLevelOutputDebugStringf(TEXT("\Existing Product Identifier: %s"), *ExistingPurchase.ToDebugString());
-		}
-	}
-
-	DECLARE_CYCLE_STAT(TEXT("FSimpleDelegateGraphTask.QueryExistingPurchases"), STAT_FSimpleDelegateGraphTask_QueryExistingPurchases, STATGROUP_TaskGraphTasks);
-
-	FSimpleDelegateGraphTask::CreateAndDispatchWhenReady(
-		FSimpleDelegateGraphTask::FDelegate::CreateLambda([=]()
-	{
-		FPlatformMisc::LowLevelOutputDebugStringf(TEXT("Query existing purchases was completed %s\n"), bWasSuccessful ? TEXT("successfully") : TEXT("unsuccessfully"));
-		if (IOnlineSubsystem* const OnlineSub = IOnlineSubsystem::Get(GOOGLEPLAY_SUBSYSTEM))
-		{
-			FOnlineSubsystemGooglePlay* const OnlineSubGP = static_cast<FOnlineSubsystemGooglePlay* const>(OnlineSub);
-			if (OnlineSubGP)
-			{
-				OnlineSubGP->TriggerOnGooglePlayQueryExistingPurchasesCompleteDelegates(EGPResponse, ExistingPurchaseInfo);
-			}
-		}
-	}),
-		GET_STATID(STAT_FSimpleDelegateGraphTask_QueryExistingPurchases),
-		nullptr,
-		ENamedThreads::GameThread
-		);
-}
diff --git a/Engine/Plugins/Online/Android/OnlineSubsystemGooglePlay/Source/Private/OnlineStoreInterfaceGooglePlay.h b/Engine/Plugins/Online/Android/OnlineSubsystemGooglePlay/Source/Private/OnlineStoreInterfaceGooglePlay.h
--- a/Engine/Plugins/Online/Android/OnlineSubsystemGooglePlay/Source/Private/OnlineStoreInterfaceGooglePlay.h
+++ b/Engine/Plugins/Online/Android/OnlineSubsystemGooglePlay/Source/Private/OnlineStoreInterfaceGooglePlay.h
@@ -61,7 +61,7 @@ private:
 	 * @param InResponseCode response from Google backend
 	 * @param AvailablePurchases list of offers returned in response to a query on available offer ids
 	 */
-	void OnGooglePlayAvailableIAPQueryComplete(EGooglePlayBillingResponseCode InResponseCode, const TArray<FInAppPurchaseProductInfo>& AvailablePurchases);
+	void OnGooglePlayAvailableIAPQueryComplete(EGooglePlayBillingResponseCode InResponseCode, const TArray<FOnlineStoreOffer>& AvailablePurchases);
 	FDelegateHandle AvailableIAPQueryDelegateHandle;
 
 	/**
diff --git a/Engine/Plugins/Online/Android/OnlineSubsystemGooglePlay/Source/Public/OnlineStoreGooglePlay.h b/Engine/Plugins/Online/Android/OnlineSubsystemGooglePlay/Source/Public/OnlineStoreGooglePlay.h
--- a/Engine/Plugins/Online/Android/OnlineSubsystemGooglePlay/Source/Public/OnlineStoreGooglePlay.h
+++ b/Engine/Plugins/Online/Android/OnlineSubsystemGooglePlay/Source/Public/OnlineStoreGooglePlay.h
@@ -70,8 +70,11 @@ private:
 	/** Reference to the parent subsystem */
 	FOnlineSubsystemGooglePlay* Subsystem;
 
-	void OnGooglePlayAvailableIAPQueryComplete(EGooglePlayBillingResponseCode InResponse, const TArray<FInAppPurchaseProductInfo>& InProvidedProductInformation);
+	void OnGooglePlayAvailableIAPQueryComplete(EGooglePlayBillingResponseCode InResponse, const TArray<FOnlineStoreOffer>& InProvidedProductInformation);
 	FDelegateHandle AvailableIAPQueryDelegateHandle;
+
+	/** Delegate fired when a query for purchases has completed, whether successful or unsuccessful */
+	FOnQueryForAvailablePurchasesComplete OnQueryForAvailablePurchasesCompleteDelegate;
 };
 
 typedef TSharedPtr<FOnlineStoreGooglePlayV2, ESPMode::ThreadSafe> FOnlineStoreGooglePlayV2Ptr;
diff --git a/Engine/Plugins/Online/Android/OnlineSubsystemGooglePlay/Source/Public/OnlineSubsystemGooglePlay.h b/Engine/Plugins/Online/Android/OnlineSubsystemGooglePlay/Source/Public/OnlineSubsystemGooglePlay.h
--- a/Engine/Plugins/Online/Android/OnlineSubsystemGooglePlay/Source/Public/OnlineSubsystemGooglePlay.h
+++ b/Engine/Plugins/Online/Android/OnlineSubsystemGooglePlay/Source/Public/OnlineSubsystemGooglePlay.h
@@ -10,6 +10,8 @@
 #include "Templates/UniquePtr.h"
 #include "HAL/RunnableThread.h"
 #include "OnlineStoreInterfaceGooglePlay.h"
+#include "OnlineStoreGooglePlay.h"
+#include "OnlinePurchaseGooglePlay.h"
 
 THIRD_PARTY_INCLUDES_START
 #include "gpg/game_services.h"
@@ -132,7 +134,7 @@ PACKAGE_SCOPE:
 	 * @param InResponseCode response from the GooglePlay backend
 	 * @param ProvidedProductInformation information returned from the backend about the queried offers
 	 */
-	DEFINE_ONLINE_DELEGATE_TWO_PARAM(OnGooglePlayAvailableIAPQueryComplete, EGooglePlayBillingResponseCode /*InResponseCode*/, const TArray<FInAppPurchaseProductInfo>& /*ProvidedProductInformation*/);
+	DEFINE_ONLINE_DELEGATE_TWO_PARAM(OnGooglePlayAvailableIAPQueryComplete, EGooglePlayBillingResponseCode /*InResponseCode*/, const TArray<FOnlineStoreOffer>& /*ProvidedProductInformation*/);
 
 	/**
 	 * Delegate fired internally when the Java in app purchase has completed, notifying any GooglePlay OSS listeners
diff --git a/Engine/Plugins/Online/OnlineSubsystem/Source/OnlineSubsystem.Build.cs b/Engine/Plugins/Online/OnlineSubsystem/Source/OnlineSubsystem.Build.cs
--- a/Engine/Plugins/Online/OnlineSubsystem/Source/OnlineSubsystem.Build.cs
+++ b/Engine/Plugins/Online/OnlineSubsystem/Source/OnlineSubsystem.Build.cs
@@ -11,7 +11,8 @@ public class OnlineSubsystem : ModuleRules
 		PublicDependencyModuleNames.AddRange(
 			new string[] {
 				"Json",
-			}
+                "OnlineSubsystemUtils"
+            }
 		);
 
 		PublicIncludePaths.Add(ModuleDirectory);
@@ -26,7 +27,7 @@ public class OnlineSubsystem : ModuleRules
 				"ImageCore",
 				"Sockets",
 				"JsonUtilities",
-				"AudioMixer"
+				"AudioMixer",
 			}
 		);
 	}
diff --git a/Engine/Plugins/Online/OnlineSubsystem/Source/Public/Interfaces/OnlineStoreInterface.h b/Engine/Plugins/Online/OnlineSubsystem/Source/Public/Interfaces/OnlineStoreInterface.h
--- a/Engine/Plugins/Online/OnlineSubsystem/Source/Public/Interfaces/OnlineStoreInterface.h
+++ b/Engine/Plugins/Online/OnlineSubsystem/Source/Public/Interfaces/OnlineStoreInterface.h
@@ -6,6 +6,7 @@
 #include "UObject/ObjectMacros.h"
 #include "OnlineSubsystemTypes.h"
 #include "OnlineDelegateMacros.h"
+#include "InAppPurchaseCallbackProxy2.h"
 #include "OnlineStoreInterface.generated.h"
 
 ONLINESUBSYSTEM_API DECLARE_LOG_CATEGORY_EXTERN(LogOnlineStore, Log, All);
diff --git a/Engine/Plugins/Online/OnlineSubsystem/Source/Public/Interfaces/OnlineStoreInterfaceV2.h b/Engine/Plugins/Online/OnlineSubsystem/Source/Public/Interfaces/OnlineStoreInterfaceV2.h
--- a/Engine/Plugins/Online/OnlineSubsystem/Source/Public/Interfaces/OnlineStoreInterfaceV2.h
+++ b/Engine/Plugins/Online/OnlineSubsystem/Source/Public/Interfaces/OnlineStoreInterfaceV2.h
@@ -3,6 +3,7 @@
 #pragma once
 
 #include "CoreMinimal.h"
+#include "OnlineDelegateMacros.h"
 
 ONLINESUBSYSTEM_API DECLARE_LOG_CATEGORY_EXTERN(LogOnlineStoreV2, Log, All);
 
@@ -22,7 +23,6 @@ typedef FString FUniqueOfferId;
 typedef FString FOfferNamespace;
 typedef FString FUniqueCategoryId;
 
-
 enum class EOnlineStoreOfferDiscountType : uint8
 {
 	/** Offer isn't on sale*/
@@ -107,6 +107,8 @@ public:
 	/** Type of discount currently running on this offer (if any) */
 	EOnlineStoreOfferDiscountType DiscountType;
 
+	TMap<FString, FString> DynamicFields;
+
 	/** @return FText suitable for localized display */
 	virtual FText GetDisplayRegularPrice() const
 	{
@@ -179,6 +181,15 @@ DECLARE_DELEGATE_TwoParams(FOnQueryOnlineStoreCategoriesComplete, bool /*bWasSuc
  */
 DECLARE_DELEGATE_ThreeParams(FOnQueryOnlineStoreOffersComplete, bool /*bWasSuccessful*/, const TArray<FUniqueOfferId>& /*OfferIds*/, const FString& /*Error*/);
 
+/**
+ * Delegate fired when a session create request has completed
+ *
+ * @param SessionName the name of the session this callback is for
+ * @param bWasSuccessful true if the async action completed without error, false if there was an error
+ */
+DECLARE_MULTICAST_DELEGATE_OneParam(FOnQueryForAvailablePurchasesComplete, bool);
+typedef FOnQueryForAvailablePurchasesComplete::FDelegate FOnQueryForAvailablePurchasesCompleteDelegate;
+
 /**
  *	Access to available offers for purchase
  */
@@ -186,6 +197,13 @@ class IOnlineStoreV2
 {
 public:
 
+	virtual ~IOnlineStoreV2(){}
+
+	/**
+	* Delegate which is executed when QueryForAvailablePurchases completes
+	*/
+	DEFINE_ONLINE_DELEGATE_ONE_PARAM(OnQueryForAvailablePurchasesComplete, bool);
+
 	/**
 	 * Query for available store categories. Delegate callback is guaranteed.
 	 *
diff --git a/Engine/Plugins/Online/OnlineSubsystem/Source/Public/Online.h b/Engine/Plugins/Online/OnlineSubsystem/Source/Public/Online.h
--- a/Engine/Plugins/Online/OnlineSubsystem/Source/Public/Online.h
+++ b/Engine/Plugins/Online/OnlineSubsystem/Source/Public/Online.h
@@ -17,6 +17,7 @@
 #include "Interfaces/OnlineTitleFileInterface.h"
 #include "Interfaces/OnlineAchievementsInterface.h"
 #include "Interfaces/OnlinePresenceInterface.h"
+#include "Interfaces/OnlinePurchaseInterface.h"
 #include "Interfaces/OnlineEntitlementsInterface.h"
 #include "Interfaces/OnlineUserCloudInterface.h"
 #include "Interfaces/OnlineUserInterface.h"
@@ -164,6 +165,13 @@ namespace Online
 	 * @return Interface pointer for the appropriate party service
 	 */
 	IMPLEMENT_GET_INTERFACE(Stats);
+
+	/**
+	 * Get the interface for accessing the purchasing services
+	 * @param SubsystemName - Name of the requested online service
+	 * @return Interface pointer for the appropriate party service
+	 */
+	IMPLEMENT_GET_INTERFACE(Purchase);
 };
 
 #undef IMPLEMENT_GET_INTERFACE
diff --git a/Engine/Plugins/Online/OnlineSubsystem/Source/Public/OnlineSubsystem.h b/Engine/Plugins/Online/OnlineSubsystem/Source/Public/OnlineSubsystem.h
--- a/Engine/Plugins/Online/OnlineSubsystem/Source/Public/OnlineSubsystem.h
+++ b/Engine/Plugins/Online/OnlineSubsystem/Source/Public/OnlineSubsystem.h
@@ -370,6 +370,7 @@ public:
 	 * Get the interface for accessing an online store
 	 * @return Interface pointer for the appropriate online store service
 	 */
+	UE_DEPRECATED(4.24, "Move to OnlineStoreInterfaceV2 and OnlinePurchaseInterface.")
 	virtual IOnlineStorePtr GetStoreInterface() const = 0;
 
 	/** 
diff --git a/Engine/Plugins/Online/OnlineSubsystemUtils/Source/OnlineBlueprintSupport/Classes/K2Node_InAppPurchase2.h b/Engine/Plugins/Online/OnlineSubsystemUtils/Source/OnlineBlueprintSupport/Classes/K2Node_InAppPurchase2.h
new file mode 100644
--- /dev/null
+++ b/Engine/Plugins/Online/OnlineSubsystemUtils/Source/OnlineBlueprintSupport/Classes/K2Node_InAppPurchase2.h
@@ -0,0 +1,14 @@
+// Copyright 1998-2019 Epic Games, Inc. All Rights Reserved.
+
+#pragma once
+
+#include "CoreMinimal.h"
+#include "UObject/ObjectMacros.h"
+#include "K2Node_BaseAsyncTask.h"
+#include "K2Node_InAppPurchase2.generated.h"
+
+UCLASS()
+class ONLINEBLUEPRINTSUPPORT_API UK2Node_InAppPurchase2 : public UK2Node_BaseAsyncTask
+{
+	GENERATED_UCLASS_BODY()
+};
diff --git a/Engine/Plugins/Online/OnlineSubsystemUtils/Source/OnlineBlueprintSupport/Classes/K2Node_InAppPurchaseQuery2.h b/Engine/Plugins/Online/OnlineSubsystemUtils/Source/OnlineBlueprintSupport/Classes/K2Node_InAppPurchaseQuery2.h
new file mode 100644
--- /dev/null
+++ b/Engine/Plugins/Online/OnlineSubsystemUtils/Source/OnlineBlueprintSupport/Classes/K2Node_InAppPurchaseQuery2.h
@@ -0,0 +1,14 @@
+// Copyright 1998-2019 Epic Games, Inc. All Rights Reserved.
+
+#pragma once
+
+#include "CoreMinimal.h"
+#include "UObject/ObjectMacros.h"
+#include "K2Node_BaseAsyncTask.h"
+#include "K2Node_InAppPurchaseQuery2.generated.h"
+
+UCLASS()
+class ONLINEBLUEPRINTSUPPORT_API UK2Node_InAppPurchaseQuery2 : public UK2Node_BaseAsyncTask
+{
+	GENERATED_UCLASS_BODY()
+};
diff --git a/Engine/Plugins/Online/OnlineSubsystemUtils/Source/OnlineBlueprintSupport/Classes/K2Node_InAppPurchaseQueryOwned2.h b/Engine/Plugins/Online/OnlineSubsystemUtils/Source/OnlineBlueprintSupport/Classes/K2Node_InAppPurchaseQueryOwned2.h
new file mode 100644
--- /dev/null
+++ b/Engine/Plugins/Online/OnlineSubsystemUtils/Source/OnlineBlueprintSupport/Classes/K2Node_InAppPurchaseQueryOwned2.h
@@ -0,0 +1,14 @@
+// Copyright 1998-2019 Epic Games, Inc. All Rights Reserved.
+
+#pragma once
+
+#include "CoreMinimal.h"
+#include "UObject/ObjectMacros.h"
+#include "K2Node_BaseAsyncTask.h"
+#include "K2Node_InAppPurchaseQueryOwned2.generated.h"
+
+UCLASS()
+class ONLINEBLUEPRINTSUPPORT_API UK2Node_InAppPurchaseQueryOwned2 : public UK2Node_BaseAsyncTask
+{
+	GENERATED_UCLASS_BODY()
+};
diff --git a/Engine/Plugins/Online/OnlineSubsystemUtils/Source/OnlineBlueprintSupport/Classes/K2Node_InAppPurchaseRestore2.h b/Engine/Plugins/Online/OnlineSubsystemUtils/Source/OnlineBlueprintSupport/Classes/K2Node_InAppPurchaseRestore2.h
new file mode 100644
--- /dev/null
+++ b/Engine/Plugins/Online/OnlineSubsystemUtils/Source/OnlineBlueprintSupport/Classes/K2Node_InAppPurchaseRestore2.h
@@ -0,0 +1,14 @@
+// Copyright 1998-2019 Epic Games, Inc. All Rights Reserved.
+
+#pragma once
+
+#include "CoreMinimal.h"
+#include "UObject/ObjectMacros.h"
+#include "K2Node_BaseAsyncTask.h"
+#include "K2Node_InAppPurchaseRestore2.generated.h"
+
+UCLASS()
+class ONLINEBLUEPRINTSUPPORT_API UK2Node_InAppPurchaseRestore2 : public UK2Node_BaseAsyncTask
+{
+	GENERATED_UCLASS_BODY()
+};
diff --git a/Engine/Plugins/Online/OnlineSubsystemUtils/Source/OnlineBlueprintSupport/Classes/K2Node_InAppPurchaseUnprocessed2.h b/Engine/Plugins/Online/OnlineSubsystemUtils/Source/OnlineBlueprintSupport/Classes/K2Node_InAppPurchaseUnprocessed2.h
new file mode 100644
--- /dev/null
+++ b/Engine/Plugins/Online/OnlineSubsystemUtils/Source/OnlineBlueprintSupport/Classes/K2Node_InAppPurchaseUnprocessed2.h
@@ -0,0 +1,14 @@
+// Copyright 1998-2019 Epic Games, Inc. All Rights Reserved.
+
+#pragma once
+
+#include "CoreMinimal.h"
+#include "UObject/ObjectMacros.h"
+#include "K2Node_BaseAsyncTask.h"
+#include "K2Node_InAppPurchaseUnprocessed2.generated.h"
+
+UCLASS()
+class ONLINEBLUEPRINTSUPPORT_API UK2Node_InAppPurchaseUnprocessed2 : public UK2Node_BaseAsyncTask
+{
+	GENERATED_UCLASS_BODY()
+};
diff --git a/Engine/Plugins/Online/OnlineSubsystemUtils/Source/OnlineBlueprintSupport/Private/K2Node_InAppPurchase2.cpp b/Engine/Plugins/Online/OnlineSubsystemUtils/Source/OnlineBlueprintSupport/Private/K2Node_InAppPurchase2.cpp
new file mode 100644
--- /dev/null
+++ b/Engine/Plugins/Online/OnlineSubsystemUtils/Source/OnlineBlueprintSupport/Private/K2Node_InAppPurchase2.cpp
@@ -0,0 +1,17 @@
+// Copyright 1998-2019 Epic Games, Inc. All Rights Reserved.
+
+#include "K2Node_InAppPurchase2.h"
+#include "InAppPurchaseCallbackProxy2.h"
+
+#define LOCTEXT_NAMESPACE "K2Node"
+
+UK2Node_InAppPurchase2::UK2Node_InAppPurchase2(const FObjectInitializer& ObjectInitializer)
+	: Super(ObjectInitializer)
+{
+	ProxyFactoryFunctionName = GET_FUNCTION_NAME_CHECKED(UInAppPurchaseCallbackProxy2, CreateProxyObjectForInAppPurchase);
+	ProxyFactoryClass = UInAppPurchaseCallbackProxy2::StaticClass();
+
+	ProxyClass = UInAppPurchaseCallbackProxy2::StaticClass();
+}
+
+#undef LOCTEXT_NAMESPACE
diff --git a/Engine/Plugins/Online/OnlineSubsystemUtils/Source/OnlineBlueprintSupport/Private/K2Node_InAppPurchaseQuery2.cpp b/Engine/Plugins/Online/OnlineSubsystemUtils/Source/OnlineBlueprintSupport/Private/K2Node_InAppPurchaseQuery2.cpp
new file mode 100644
--- /dev/null
+++ b/Engine/Plugins/Online/OnlineSubsystemUtils/Source/OnlineBlueprintSupport/Private/K2Node_InAppPurchaseQuery2.cpp
@@ -0,0 +1,17 @@
+// Copyright 1998-2019 Epic Games, Inc. All Rights Reserved.
+
+#include "K2Node_InAppPurchaseQuery2.h"
+#include "InAppPurchaseQueryCallbackProxy2.h"
+
+#define LOCTEXT_NAMESPACE "K2Node"
+
+UK2Node_InAppPurchaseQuery2::UK2Node_InAppPurchaseQuery2(const FObjectInitializer& ObjectInitializer)
+	: Super(ObjectInitializer)
+{
+	ProxyFactoryFunctionName = GET_FUNCTION_NAME_CHECKED(UInAppPurchaseQueryCallbackProxy2, CreateProxyObjectForInAppPurchaseQuery);
+	ProxyFactoryClass = UInAppPurchaseQueryCallbackProxy2::StaticClass();
+
+	ProxyClass = UInAppPurchaseQueryCallbackProxy2::StaticClass();
+}
+
+#undef LOCTEXT_NAMESPACE
diff --git a/Engine/Plugins/Online/OnlineSubsystemUtils/Source/OnlineBlueprintSupport/Private/K2Node_InAppPurchaseQueryOwned2.cpp b/Engine/Plugins/Online/OnlineSubsystemUtils/Source/OnlineBlueprintSupport/Private/K2Node_InAppPurchaseQueryOwned2.cpp
new file mode 100644
--- /dev/null
+++ b/Engine/Plugins/Online/OnlineSubsystemUtils/Source/OnlineBlueprintSupport/Private/K2Node_InAppPurchaseQueryOwned2.cpp
@@ -0,0 +1,17 @@
+// Copyright 1998-2019 Epic Games, Inc. All Rights Reserved.
+
+#include "K2Node_InAppPurchaseQueryOwned2.h"
+#include "InAppPurchaseCallbackProxy2.h"
+
+#define LOCTEXT_NAMESPACE "K2Node"
+
+UK2Node_InAppPurchaseQueryOwned2::UK2Node_InAppPurchaseQueryOwned2(const FObjectInitializer& ObjectInitializer)
+	: Super(ObjectInitializer)
+{
+	ProxyFactoryFunctionName = GET_FUNCTION_NAME_CHECKED(UInAppPurchaseCallbackProxy2, CreateProxyObjectForInAppPurchaseQueryOwned);
+	ProxyFactoryClass = UInAppPurchaseCallbackProxy2::StaticClass();
+
+	ProxyClass = UInAppPurchaseCallbackProxy2::StaticClass();
+}
+
+#undef LOCTEXT_NAMESPACE
diff --git a/Engine/Plugins/Online/OnlineSubsystemUtils/Source/OnlineBlueprintSupport/Private/K2Node_InAppPurchaseRestore2.cpp b/Engine/Plugins/Online/OnlineSubsystemUtils/Source/OnlineBlueprintSupport/Private/K2Node_InAppPurchaseRestore2.cpp
new file mode 100644
--- /dev/null
+++ b/Engine/Plugins/Online/OnlineSubsystemUtils/Source/OnlineBlueprintSupport/Private/K2Node_InAppPurchaseRestore2.cpp
@@ -0,0 +1,17 @@
+// Copyright 1998-2019 Epic Games, Inc. All Rights Reserved.
+
+#include "K2Node_InAppPurchaseRestore2.h"
+#include "InAppPurchaseRestoreCallbackProxy2.h"
+
+#define LOCTEXT_NAMESPACE "K2Node"
+
+UK2Node_InAppPurchaseRestore2::UK2Node_InAppPurchaseRestore2(const FObjectInitializer& ObjectInitializer)
+	: Super(ObjectInitializer)
+{
+	ProxyFactoryFunctionName = GET_FUNCTION_NAME_CHECKED(UInAppPurchaseRestoreCallbackProxy2, CreateProxyObjectForInAppPurchaseRestore);
+	ProxyFactoryClass = UInAppPurchaseRestoreCallbackProxy2::StaticClass();
+
+	ProxyClass = UInAppPurchaseRestoreCallbackProxy2::StaticClass();
+}
+
+#undef LOCTEXT_NAMESPACE
diff --git a/Engine/Plugins/Online/OnlineSubsystemUtils/Source/OnlineBlueprintSupport/Private/K2Node_InAppPurchaseUnprocessed2.cpp b/Engine/Plugins/Online/OnlineSubsystemUtils/Source/OnlineBlueprintSupport/Private/K2Node_InAppPurchaseUnprocessed2.cpp
new file mode 100644
--- /dev/null
+++ b/Engine/Plugins/Online/OnlineSubsystemUtils/Source/OnlineBlueprintSupport/Private/K2Node_InAppPurchaseUnprocessed2.cpp
@@ -0,0 +1,17 @@
+// Copyright 1998-2019 Epic Games, Inc. All Rights Reserved.
+
+#include "K2Node_InAppPurchaseUnprocessed2.h"
+#include "InAppPurchaseCallbackProxy2.h"
+
+#define LOCTEXT_NAMESPACE "K2Node"
+
+UK2Node_InAppPurchaseUnprocessed2::UK2Node_InAppPurchaseUnprocessed2(const FObjectInitializer& ObjectInitializer)
+	: Super(ObjectInitializer)
+{
+	ProxyFactoryFunctionName = GET_FUNCTION_NAME_CHECKED(UInAppPurchaseCallbackProxy2, CreateProxyObjectForInAppPurchaseUnprocessedPurchases);
+	ProxyFactoryClass = UInAppPurchaseCallbackProxy2::StaticClass();
+
+	ProxyClass = UInAppPurchaseCallbackProxy2::StaticClass();
+}
+
+#undef LOCTEXT_NAMESPACE
diff --git a/Engine/Plugins/Online/OnlineSubsystemUtils/Source/OnlineSubsystemUtils/Classes/InAppPurchaseCallbackProxy.h b/Engine/Plugins/Online/OnlineSubsystemUtils/Source/OnlineSubsystemUtils/Classes/InAppPurchaseCallbackProxy.h
--- a/Engine/Plugins/Online/OnlineSubsystemUtils/Source/OnlineSubsystemUtils/Classes/InAppPurchaseCallbackProxy.h
+++ b/Engine/Plugins/Online/OnlineSubsystemUtils/Source/OnlineSubsystemUtils/Classes/InAppPurchaseCallbackProxy.h
@@ -8,7 +8,28 @@
 #include "Interfaces/OnlineStoreInterface.h"
 #include "InAppPurchaseCallbackProxy.generated.h"
 
-DECLARE_DYNAMIC_MULTICAST_DELEGATE_TwoParams(FInAppPurchaseResult, EInAppPurchaseState::Type, CompletionStatus, const FInAppPurchaseProductInfo&, InAppPurchaseInformation);
+/**
+ * Micro-transaction purchase information
+ */
+USTRUCT(BlueprintType)
+struct FInAppPurchaseReceiptInfo
+{
+	GENERATED_USTRUCT_BODY()
+
+		// The item name
+		UPROPERTY(BlueprintReadOnly, Category = ProductInfo)
+		FString ItemName;
+
+	// The unique product identifier
+	UPROPERTY(BlueprintReadOnly, Category = ProductInfo)
+		FString ItemId;
+
+	// the unique transaction identifier
+	UPROPERTY(BlueprintReadOnly, Category = ProductInfo)
+		FString ValidationInfo;
+};
+
+DECLARE_DYNAMIC_MULTICAST_DELEGATE_TwoParams(FInAppPurchaseResult, EInAppPurchaseState::Type, PurchaseStatus, const FInAppPurchaseProductInfo&, InAppPurchaseReceipts);
 
 UCLASS(MinimalAPI)
 class UInAppPurchaseCallbackProxy : public UObject
diff --git a/Engine/Plugins/Online/OnlineSubsystemUtils/Source/OnlineSubsystemUtils/Public/OnlineSubsystemUtils.h b/Engine/Plugins/Online/OnlineSubsystemUtils/Source/OnlineSubsystemUtils/Public/OnlineSubsystemUtils.h
--- a/Engine/Plugins/Online/OnlineSubsystemUtils/Source/OnlineSubsystemUtils/Public/OnlineSubsystemUtils.h
+++ b/Engine/Plugins/Online/OnlineSubsystemUtils/Source/OnlineSubsystemUtils/Public/OnlineSubsystemUtils.h
@@ -300,6 +300,7 @@ namespace Online
 	IMPLEMENT_GET_INTERFACE(Presence);
 	IMPLEMENT_GET_INTERFACE(Events);
 	IMPLEMENT_GET_INTERFACE(Stats);
+	IMPLEMENT_GET_INTERFACE(Purchase);
 }
 
 #undef IMPLEMENT_GET_INTERFACE
