Fix possible crash for missing audio manager properties

#ue4
#android

diff --git a/Engine/Build/Android/Java/src/com/epicgames/ue4/GameActivity.java.template b/Engine/Build/Android/Java/src/com/epicgames/ue4/GameActivity.java.template
--- a/Engine/Build/Android/Java/src/com/epicgames/ue4/GameActivity.java.template
+++ b/Engine/Build/Android/Java/src/com/epicgames/ue4/GameActivity.java.template
@@ -5290,7 +5290,7 @@ public class GameActivity extends $${gameActivitySuperClass}$$ implements Surfac
 		{
 			AudioManager am = (AudioManager) getSystemService(Context.AUDIO_SERVICE);
 			String framesPerBuffer = am.getProperty(AudioManager.PROPERTY_OUTPUT_FRAMES_PER_BUFFER);
-			int framesPerBufferInt = Integer.parseInt(framesPerBuffer);
+			int framesPerBufferInt = framesPerBuffer == null ? 0 : Integer.parseInt(framesPerBuffer);
 			if (framesPerBufferInt == 0) framesPerBufferInt = 256; // Use default
 			Log.debug("[JAVA] audiomanager.framesPerBuffer = " + framesPerBufferInt);
 			return framesPerBufferInt;
@@ -5300,7 +5300,7 @@ public class GameActivity extends $${gameActivitySuperClass}$$ implements Surfac
 		{
 			AudioManager am = (AudioManager) getSystemService(Context.AUDIO_SERVICE);
 			String sampleRateStr = am.getProperty(AudioManager.PROPERTY_OUTPUT_SAMPLE_RATE);
-			int sampleRateInt = Integer.parseInt(sampleRateStr);
+			int sampleRateInt = sampleRateStr == null ? 0 : Integer.parseInt(sampleRateStr);
 			if (sampleRateInt == 0) sampleRateInt = 44100; // Use a default value if property not found
 			Log.debug("[JAVA] audiomanager.optimalSampleRate = " + sampleRateInt);
 			return sampleRateInt;
