Fix compile problem for Android

#ue4
#android

diff --git a/Engine/Build/Android/Java/src/com/epicgames/ue4/GameActivity.java.template b/Engine/Build/Android/Java/src/com/epicgames/ue4/GameActivity.java.template
--- a/Engine/Build/Android/Java/src/com/epicgames/ue4/GameActivity.java.template
+++ b/Engine/Build/Android/Java/src/com/epicgames/ue4/GameActivity.java.template
@@ -5132,7 +5132,7 @@ public class GameActivity extends $${gameActivitySuperClass}$$ implements Surfac
 
 	public boolean AndroidThunkJava_IsAllowedRemoteNotifications()
 	{
-		bool bPluginEnabled = false;
+		boolean bPluginEnabled = false;
 //$${gameActivityAllowedRemoteNotificationsAdditions}$$
 		return bPluginEnabled && NotificationManagerCompat.from(getApplicationContext()).areNotificationsEnabled();
 	}
