Update NetworkChangedManager to provide connection transport type

The networkChangedManager is the best place to know what kind of connection type we have. This updates the engine and both related projects (web and native) to support receiving this information in the onNetworkAvailable callback.
The callback will be called each time the transport type changes.

Test-plan:
Built native and web versions of app here:
Native:
https://ec-01.epicgames.net/commander/link/jobDetails/jobs/10035906?
Web:
https://ec-01.epicgames.net/commander/link/jobDetails/jobs/10035907?
Launched both versions of the app
Verified that turning off wifi and back on triggered no connection alert
Verified that switching between wifi and cellular back triggered calls to the onNetworkAvailable as expected

diff --git a/Engine/Build/Android/Java/src/com/epicgames/ue4/network/NetworkChangedManager.java b/Engine/Build/Android/Java/src/com/epicgames/ue4/network/NetworkChangedManager.java
--- a/Engine/Build/Android/Java/src/com/epicgames/ue4/network/NetworkChangedManager.java
+++ b/Engine/Build/Android/Java/src/com/epicgames/ue4/network/NetworkChangedManager.java
@@ -63,6 +63,8 @@ public final class NetworkChangedManager implements NetworkConnectivityClient {
 	private Set<String> networks = new HashSet<>();
 	@Nullable
 	private ConnectivityState currentState = null;
+	@Nullable
+	private NetworkTransportType currentNetworkTransport = NetworkTransportType.UNKNOWN;
 
 	private boolean networkCheckInProgress = false;
 	private boolean retryScheduled = false;
@@ -87,6 +89,11 @@ public final class NetworkChangedManager implements NetworkConnectivityClient {
 		}
 		NetworkRequest.Builder builder = new NetworkRequest.Builder();
 		builder.addCapability(NetworkCapabilities.NET_CAPABILITY_INTERNET);
+		builder.addTransportType(NetworkCapabilities.TRANSPORT_WIFI);
+		builder.addTransportType(NetworkCapabilities.TRANSPORT_ETHERNET);
+		builder.addTransportType(NetworkCapabilities.TRANSPORT_VPN);
+		builder.addTransportType(NetworkCapabilities.TRANSPORT_CELLULAR);
+		builder.addTransportType(NetworkCapabilities.TRANSPORT_BLUETOOTH);
 		if (Build.VERSION.SDK_INT >= Build.VERSION_CODES.M) {
 			builder.addCapability(NetworkCapabilities.NET_CAPABILITY_VALIDATED);
 		}
@@ -170,7 +177,8 @@ public final class NetworkChangedManager implements NetworkConnectivityClient {
 	}
 
 	private void setNetworkState(ConnectivityState state) {
-		if (currentState == state) {
+        NetworkTransportType updatedNetworkTransportType = calculateNetworkTransport(connectivityManager);
+        if (currentState == state && currentNetworkTransport == updatedNetworkTransportType) {
 			Log.verbose("Connectivity hasn't changed. Current state: " + currentState);
 			if (currentState != ConnectivityState.CONNECTION_AVAILABLE) {
 				scheduleRetry();
@@ -178,7 +186,8 @@ public final class NetworkChangedManager implements NetworkConnectivityClient {
 			return;
 		}
 		currentState = state;
-		fireNetworkChangeListeners(state);
+        currentNetworkTransport = updatedNetworkTransportType;
+		fireNetworkChangeListeners(state, currentNetworkTransport);
 		Log.verbose("Network connectivity changed. New connectivity state: " + state);
 
 		if (currentState != ConnectivityState.CONNECTION_AVAILABLE) {
@@ -297,7 +306,7 @@ public final class NetworkChangedManager implements NetworkConnectivityClient {
 		}
 		// Current state will never be null after calling checkNetworkConnectivity
 		if (fireImmediately && currentState != null) {
-			fireNetworkChangeListenerInternal(listener, currentState);
+			fireNetworkChangeListenerInternal(listener, currentState, currentNetworkTransport);
 		}
 		return true;
 	}
@@ -327,7 +336,7 @@ public final class NetworkChangedManager implements NetworkConnectivityClient {
 		return false;
 	}
 
-	private void fireNetworkChangeListeners(ConnectivityState state) {
+	private void fireNetworkChangeListeners(ConnectivityState state, NetworkTransportType networkTransportType) {
 		Iterator<WeakReference<Listener>> changeListenersIterator = networkChangedListeners.iterator();
 		while (changeListenersIterator.hasNext()) {
 			WeakReference<Listener> listenerWeakReference = changeListenersIterator.next();
@@ -335,18 +344,18 @@ public final class NetworkChangedManager implements NetworkConnectivityClient {
 			if (listener == null) {
 				changeListenersIterator.remove();
 			} else {
-				fireNetworkChangeListenerInternal(listener, state);
+				fireNetworkChangeListenerInternal(listener, state, networkTransportType);
 			}
 		}
 	}
 
-	private void fireNetworkChangeListenerInternal(Listener listener, ConnectivityState state) {
+	private void fireNetworkChangeListenerInternal(Listener listener, ConnectivityState state, NetworkTransportType networkTransportType) {
 		switch (state) {
 			case NO_CONNECTION:
 				listener.onNetworkLost();
 				break;
 			case CONNECTION_AVAILABLE:
-				listener.onNetworkAvailable();
+				listener.onNetworkAvailable(networkTransportType);
 				break;
 		}
 	}
@@ -355,4 +364,51 @@ public final class NetworkChangedManager implements NetworkConnectivityClient {
 	public void checkConnectivity() {
 		checkNetworkConnectivity();
 	}
+
+    private NetworkTransportType calculateNetworkTransport(@NonNull ConnectivityManager connectivityManager) {
+       if (Build.VERSION.SDK_INT < Build.VERSION_CODES.M) {
+           NetworkInfo networkInfo = connectivityManager.getActiveNetworkInfo();
+           if (networkInfo == null) {
+               return NetworkTransportType.UNKNOWN;
+           }
+
+           switch (networkInfo.getType()) {
+               case ConnectivityManager.TYPE_WIFI:
+                   return NetworkTransportType.WIFI;
+               case ConnectivityManager.TYPE_MOBILE:
+                   return NetworkTransportType.CELLULAR;
+               case ConnectivityManager.TYPE_ETHERNET:
+                   return NetworkTransportType.ETHERNET;
+               case ConnectivityManager.TYPE_BLUETOOTH:
+                   return NetworkTransportType.BLUETOOTH;
+               case ConnectivityManager.TYPE_VPN:
+                   return NetworkTransportType.VPN;
+               default:
+                   return NetworkTransportType.UNKNOWN;
+           }
+       } else {
+           Network activeNetwork = connectivityManager.getActiveNetwork();
+           if (activeNetwork == null) {
+               return NetworkTransportType.UNKNOWN;
+           }
+           NetworkCapabilities capabilities = connectivityManager.getNetworkCapabilities(activeNetwork);
+           if (capabilities == null) {
+               return NetworkTransportType.UNKNOWN;
+           }
+
+           if (capabilities.hasTransport(NetworkCapabilities.TRANSPORT_WIFI)) {
+               return NetworkTransportType.WIFI;
+           } else if (capabilities.hasTransport(NetworkCapabilities.TRANSPORT_CELLULAR)) {
+               return NetworkTransportType.CELLULAR;
+           } else if (capabilities.hasTransport(NetworkCapabilities.TRANSPORT_ETHERNET)) {
+               return NetworkTransportType.ETHERNET;
+           } else if (capabilities.hasTransport(NetworkCapabilities.TRANSPORT_BLUETOOTH)) {
+               return NetworkTransportType.BLUETOOTH;
+           } else if (capabilities.hasTransport(NetworkCapabilities.TRANSPORT_VPN)) {
+               return NetworkTransportType.VPN;
+           } else {
+               return NetworkTransportType.UNKNOWN;
+           }
+       }
+   }
 }
diff --git a/Engine/Build/Android/Java/src/com/epicgames/ue4/network/NetworkConnectivityClient.java b/Engine/Build/Android/Java/src/com/epicgames/ue4/network/NetworkConnectivityClient.java
--- a/Engine/Build/Android/Java/src/com/epicgames/ue4/network/NetworkConnectivityClient.java
+++ b/Engine/Build/Android/Java/src/com/epicgames/ue4/network/NetworkConnectivityClient.java
@@ -6,8 +6,17 @@ import android.content.Context;
 import android.support.annotation.NonNull;
 
 public interface NetworkConnectivityClient {
-	public interface Listener {
-		void onNetworkAvailable();
+    public enum NetworkTransportType {
+        WIFI,
+        VPN,
+        ETHERNET,
+        CELLULAR,
+        BLUETOOTH,
+        UNKNOWN
+    }
+
+    public interface Listener {
+        void onNetworkAvailable(NetworkTransportType networkTransportType);
 
 		void onNetworkLost();
 	}
