Disable autofill manager in GameActivity to prevent Android 10 crash

#android

diff --git a/Engine/Build/Android/Java/src/com/epicgames/ue4/GameActivity.java.template b/Engine/Build/Android/Java/src/com/epicgames/ue4/GameActivity.java.template
--- a/Engine/Build/Android/Java/src/com/epicgames/ue4/GameActivity.java.template
+++ b/Engine/Build/Android/Java/src/com/epicgames/ue4/GameActivity.java.template
@@ -2232,6 +2232,12 @@ public class GameActivity extends $${gameActivitySuperClass}$$ implements Surfac
 		InternalFilesDir = getFilesDir().getAbsolutePath() + "/";
 		ExternalFilesDir = getExternalFilesDir(null).getAbsolutePath() + "/";
 
+		// disable autofill to avoid possible Android 10 crash (android.view.autofill.AutofillManager.ensureServiceClientAddedIfNeededLocked() SyncResultReceiver$TimeoutException)
+		if (Build.VERSION.SDK_INT >= Build.VERSION_CODES.O)
+		{
+			getWindow().getDecorView().setImportantForAutofill(View.IMPORTANT_FOR_AUTOFILL_NO_EXCLUDE_DESCENDANTS);
+		}
+
 		// create splashscreen dialog (if launched by SplashActivity)
 		_extrasBundle = getIntent().getExtras();
 		if (_extrasBundle != null)
