Try to fix an android input issue with floating keyboard.

diff --git a/Engine/Build/Android/Java/src/com/epicgames/ue4/GameActivity.java.template b/Engine/Build/Android/Java/src/com/epicgames/ue4/GameActivity.java.template
--- a/Engine/Build/Android/Java/src/com/epicgames/ue4/GameActivity.java.template
+++ b/Engine/Build/Android/Java/src/com/epicgames/ue4/GameActivity.java.template
@@ -3035,9 +3035,14 @@ public class GameActivity extends $${gameActivitySuperClass}$$ implements Surfac
                     int visibleScreenYOffset = Math.max(bottomDiff, topDiff);
 
 					//Log.debug("VK: show?" + visibleScreenYOffset + "," + newVirtualKeyboardInput.getY());
+					
+					// If the keyboard is a soft keyboard, the visibleScreenYOffset should be greater than 200 tall
+					boolean bSoftKeyboardPopup = visibleScreenYOffset > 200;
+
+					// If the keyboard is a floating keyboard, after testing on partial Samsung devcies and Huawei devices, the visibleScreenYOffset is 1 pixel tall, expand it to [1, 200] for safety, correct this range if there is any unexpected case.
+					boolean bFloatingKeyboardPopup = visibleScreenYOffset >= 1 && visibleScreenYOffset <= 200;
 
-					// Samsung Game Tools floating keyboard is only 1 pixel tall
-					if(visibleScreenYOffset > 200 || (bSamsungDevice && visibleScreenYOffset > 0))
+					if(bSoftKeyboardPopup || bFloatingKeyboardPopup)
                     {
 						if (bKeyboardShowing)
 						{
