Fix an android input issue with floating keyboard

diff --git a/Engine/Build/Android/Java/src/com/epicgames/ue4/GameActivity.java.template b/Engine/Build/Android/Java/src/com/epicgames/ue4/GameActivity.java.template
--- a/Engine/Build/Android/Java/src/com/epicgames/ue4/GameActivity.java.template
+++ b/Engine/Build/Android/Java/src/com/epicgames/ue4/GameActivity.java.template
@@ -3039,10 +3039,10 @@ public class GameActivity extends $${gameActivitySuperClass}$$ implements Surfac
 					// If the keyboard is a soft keyboard, the visibleScreenYOffset should be greater than 200 tall
 					boolean bSoftKeyboardPopup = visibleScreenYOffset > 200;
 
-					// If the keyboard is a floating keyboard, after testing on partial Samsung devcies and Huawei devices, the visibleScreenYOffset is 1 pixel tall, expand it to [1, 200] for safety, correct this range if there is any unexpected case.
-					boolean bFloatingKeyboardPopup = visibleScreenYOffset >= 1 && visibleScreenYOffset <= 200;
+					// If the keyboard is a floating keyboard, after testing with some Chinese input methods the visibleScreenYOffset can be 0 pixel tall, expand it to [0, 200] for safety, correct this range if there is any unexpected case.
+					boolean bFloatingKeyboardPopup = visibleScreenYOffset >= 0 && visibleScreenYOffset <= 200;
 
-					if(bSoftKeyboardPopup || bFloatingKeyboardPopup)
+					if(bSoftKeyboardPopup || bFloatingKeyboardPopup) // actually this condition equals to the bKeyboardShowing, keep it in case we need to handle the different keyboards seperately.
                     {
 						if (bKeyboardShowing)
 						{
