Catch exception that can occur when attempting to clear keyboard focus after having used autocomplete

Importing changes applied to 12.10 here: https://p4-swarm.epicgames.net/reviews/12121724/

Test Plan:
Launched app and game without fix
Verified that after typing in text, using auto complete, and dismissing the keyboard the game did crash

Launched app and game with fix
Verified that after typing in text, using auto complete, and dismissing the keyboard the game did not crash

This was tested and verified by QA in 12.10. See this review: https://p4-swarm.epicgames.net/reviews/12121724/ and the referenced jira task.

diff --git a/Engine/Build/Android/Java/src/com/epicgames/ue4/GameActivity.java.template b/Engine/Build/Android/Java/src/com/epicgames/ue4/GameActivity.java.template
--- a/Engine/Build/Android/Java/src/com/epicgames/ue4/GameActivity.java.template
+++ b/Engine/Build/Android/Java/src/com/epicgames/ue4/GameActivity.java.template
@@ -3992,7 +3992,13 @@ public class GameActivity extends $${gameActivitySuperClass}$$ implements Surfac
 						{
 							//Log.debug("VK: Hide newVirtualKeyboardInput");
 
-							newVirtualKeyboardInput.clearFocus();
+							if (newVirtualKeyboardInput.hasFocus()) {
+								try {
+									newVirtualKeyboardInput.clearFocus();
+								} catch (Exception e) {
+									Log.warn("Unable to clear focus from virtualKeyboardInput");
+								}
+							}
 							//set offscreen
 							newVirtualKeyboardInput.setY(-1000);
 
