Android java app restart only happens from the uithread.
 - Fixes possible race condition in which startActivity could create the new activity in the same process as the outgoing activity and then exiting process.

diff --git a/Engine/Build/Android/Java/src/com/epicgames/ue4/GameActivity.java.template b/Engine/Build/Android/Java/src/com/epicgames/ue4/GameActivity.java.template
--- a/Engine/Build/Android/Java/src/com/epicgames/ue4/GameActivity.java.template
+++ b/Engine/Build/Android/Java/src/com/epicgames/ue4/GameActivity.java.template
@@ -5985,7 +5985,7 @@ public class GameActivity extends $${gameActivitySuperClass}$$ implements Surfac
 		return true;
 	}
 
-	public void AndroidThunkJava_RestartApplication(String RestartExtra) 
+	private void RestartApplication(String RestartExtra)
 	{
 		Context context = getApplicationContext();
 	
@@ -5999,9 +5999,35 @@ public class GameActivity extends $${gameActivitySuperClass}$$ implements Surfac
 		context.startActivity(intent);
 
 		// kill the current application instance. Undefined behaviour after this!
-//$${gameActivityOnRestartApplicationAdditions}$$
 		Runtime.getRuntime().exit(0); 
 	}
+
+	public void AndroidThunkJava_RestartApplication(final String RestartExtra) 
+	{
+	    Runnable restartRunnableUIThread = new Runnable()
+		{
+			@Override
+			public void run()
+			{
+				RestartApplication(RestartExtra);
+			}
+		};
+		Log.debug("app restart : " + RestartExtra);
+//$${gameActivityOnRestartApplicationAdditions}$$
+		runOnUiThread(restartRunnableUIThread);
+
+		// wait for exit.
+		while(true)
+		{
+			try 
+			{
+				Thread.sleep(Long.MAX_VALUE); 
+			}
+			catch (InterruptedException e) {
+				e.printStackTrace();
+			}
+		}
+	}
 	
 	public void AndroidThunkJava_ClipboardCopy(String text) 
 	{
