Center android splash image to avoid stretching on wide screens

diff --git a/Engine/Build/Android/Java/src/com/epicgames/ue4/GameActivity.java.template b/Engine/Build/Android/Java/src/com/epicgames/ue4/GameActivity.java.template
--- a/Engine/Build/Android/Java/src/com/epicgames/ue4/GameActivity.java.template
+++ b/Engine/Build/Android/Java/src/com/epicgames/ue4/GameActivity.java.template
@@ -203,6 +203,7 @@ import android.view.inputmethod.InputConnection;
 import android.view.inputmethod.InputConnectionWrapper;
 import android.util.AttributeSet;
 import android.graphics.Color;
+import android.graphics.drawable.BitmapDrawable;
 import android.widget.TextView;
 import android.widget.TextView.OnEditorActionListener;
 import android.content.BroadcastReceiver;
@@ -2269,6 +2270,15 @@ public class GameActivity extends $${gameActivitySuperClass}$$ implements Surfac
 														| View.SYSTEM_UI_FLAG_IMMERSIVE);  // NOT sticky.. will be set to sticky later! 
 						}
 					}
+					
+					// Centering background image by default
+					View decorView = mSplashDialog.getWindow().getDecorView();
+					BitmapDrawable SplashBitmap = (BitmapDrawable)decorView.getBackground();
+					if (SplashBitmap != null)
+					{
+						SplashBitmap.setGravity(android.view.Gravity.CENTER);
+					}
+					
 					mSplashDialog.show();
 				}
 				catch (Exception e) {
diff --git a/Engine/Build/Android/Java/src/com/epicgames/ue4/SplashActivity.java b/Engine/Build/Android/Java/src/com/epicgames/ue4/SplashActivity.java
--- a/Engine/Build/Android/Java/src/com/epicgames/ue4/SplashActivity.java
+++ b/Engine/Build/Android/Java/src/com/epicgames/ue4/SplashActivity.java
@@ -21,6 +21,7 @@ import android.provider.Settings;
 import android.view.View;
 import android.view.WindowManager;
 import android.net.Uri;
+import android.graphics.drawable.BitmapDrawable;
 
 import android.support.v4.app.ActivityCompat;
 import android.support.v4.content.ContextCompat;
@@ -95,6 +96,14 @@ public class SplashActivity extends Activity
 											| View.SYSTEM_UI_FLAG_IMMERSIVE);  // NOT sticky (will be set later in MainActivity)
 			}
 		}
+
+		// Centering background image by default
+		View decorView = getWindow().getDecorView();
+		BitmapDrawable SplashBitmap = (BitmapDrawable)decorView.getBackground();
+		if (SplashBitmap != null)
+		{
+			SplashBitmap.setGravity(android.view.Gravity.CENTER);
+		}
 		
 		// allow certain models for now to use full area around cutout
 		boolean BlockDisplayCutout = true;
