adding android screenshot blocking, adding iOS screenshot/video caption detection, enabling video analytics for android

diff --git a/Engine/Build/Android/Java/src/com/epicgames/ue4/GameActivity.java.template b/Engine/Build/Android/Java/src/com/epicgames/ue4/GameActivity.java.template
--- a/Engine/Build/Android/Java/src/com/epicgames/ue4/GameActivity.java.template
+++ b/Engine/Build/Android/Java/src/com/epicgames/ue4/GameActivity.java.template
@@ -350,6 +350,9 @@ public class GameActivity extends $${gameActivitySuperClass}$$ implements Surfac
 	// remember screensaver state to restore
 	private boolean bKeepScreenOn = false;
 	
+	// number of times screen capturing has been disabled
+	private int NumTimesScreenCaptureDisabled = 0;
+
 	// Virtual keyboard
 	AlertDialog virtualKeyboardAlert;
 	EditText virtualKeyboardInputBox;
@@ -3321,6 +3324,9 @@ public class GameActivity extends $${gameActivitySuperClass}$$ implements Surfac
 		// restore screensaver state
 		AndroidThunkJava_KeepScreenOn(bKeepScreenOn);
 		
+		// restore screep capture state
+		SetDisableScreenCaptureInternal(AndroidThunkJava_IsScreenCaptureDisabled());
+
 		if (bForceGameEndWithError || bForceGameEndWithWarning)
 		{
 			final boolean bWarning = bForceGameEndWithWarning;
@@ -3742,6 +3748,50 @@ public class GameActivity extends $${gameActivitySuperClass}$$ implements Surfac
 		}
 	}
 
+	public boolean AndroidThunkJava_IsScreenCaptureDisabled()
+	{
+		return NumTimesScreenCaptureDisabled != 0;
+	}
+
+	private void SetDisableScreenCaptureInternal(boolean bDisable)
+	{
+		if (bDisable)
+		{
+			_activity.runOnUiThread(new Runnable()
+			{
+				@Override
+				public void run()
+				{
+					Log.debug("==============> [JAVA] AndroidThunkJava_DisableScreenCapture(true) - Disabled screen captures");
+					_activity.getWindow().addFlags(WindowManager.LayoutParams.FLAG_SECURE);
+				}
+			});
+		}
+		else
+		{
+			_activity.runOnUiThread(new Runnable()
+			{
+				@Override
+				public void run()
+				{
+					Log.debug("==============> [JAVA] AndroidThunkJava_DisableScreenCapture(false) - Enabled screen captures");
+					_activity.getWindow().clearFlags(WindowManager.LayoutParams.FLAG_SECURE);
+				}
+			});
+		}
+	}
+
+	public void AndroidThunkJava_DisableScreenCapture(boolean bDisable)
+	{
+		boolean bWasDisabled = AndroidThunkJava_IsScreenCaptureDisabled();
+		NumTimesScreenCaptureDisabled += bDisable ? 1 : -1;
+		// Change in state?
+		if (AndroidThunkJava_IsScreenCaptureDisabled() != bWasDisabled)
+		{
+			SetDisableScreenCaptureInternal(!bWasDisabled);
+		}
+	}
+
 	private class VibrateRunnable implements Runnable {
 		private int duration;
 		private Vibrator vibrator;
diff --git a/Engine/Source/Runtime/Launch/Private/Android/AndroidJNI.cpp b/Engine/Source/Runtime/Launch/Private/Android/AndroidJNI.cpp
--- a/Engine/Source/Runtime/Launch/Private/Android/AndroidJNI.cpp
+++ b/Engine/Source/Runtime/Launch/Private/Android/AndroidJNI.cpp
@@ -118,6 +118,11 @@ void FJavaWrapper::FindClassesAndMethods(JNIEnv* Env)
 	AndroidThunkJava_GetIntentExtrasString = FindMethod(Env, GameActivityClassID, "AndroidThunkJava_GetIntentExtrasString", "(Ljava/lang/String;)Ljava/lang/String;", bIsOptional);
 	AndroidThunkJava_PushSensorEvents = FindMethod(Env, GameActivityClassID, "AndroidThunkJava_PushSensorEvents", "()V", bIsOptional);
 
+	// Screen capture/recording permission
+	AndroidThunkJava_IsScreenCaptureDisabled = FindMethod(Env, GameActivityClassID, "AndroidThunkJava_IsScreenCaptureDisabled", "()Z", bIsOptional);
+	AndroidThunkJava_DisableScreenCapture = FindMethod(Env, GameActivityClassID, "AndroidThunkJava_DisableScreenCapture", "(Z)V", bIsOptional);
+
+
 	// this is optional - only inserted if Oculus Mobile plugin enabled
 	AndroidThunkJava_IsOculusMobileApplication = FindMethod(Env, GameActivityClassID, "AndroidThunkJava_IsOculusMobileApplication", "()Z", true);
 
@@ -398,6 +403,8 @@ jmethodID FJavaWrapper::AndroidThunkJava_GetIntentExtrasInt;
 jmethodID FJavaWrapper::AndroidThunkJava_GetIntentExtrasString;
 jmethodID FJavaWrapper::AndroidThunkJava_SetSustainedPerformanceMode;
 jmethodID FJavaWrapper::AndroidThunkJava_PushSensorEvents;
+jmethodID FJavaWrapper::AndroidThunkJava_IsScreenCaptureDisabled;
+jmethodID FJavaWrapper::AndroidThunkJava_DisableScreenCapture;
 
 jclass FJavaWrapper::InputDeviceInfoClass;
 jfieldID FJavaWrapper::InputDeviceInfo_VendorId;
@@ -753,6 +760,25 @@ void AndroidThunkCpp_PushSensorEvents()
 	}
 }
 
+bool AndroidThunkCpp_IsScreenCaptureDisabled()
+{
+	if (JNIEnv* Env = FAndroidApplication::GetJavaEnv())
+	{
+		// call the java side
+		return FJavaWrapper::CallBooleanMethod(Env, FJavaWrapper::GameActivityThis, FJavaWrapper::AndroidThunkJava_IsScreenCaptureDisabled);
+	}
+	return true;
+}
+
+void AndroidThunkCpp_DisableScreenCapture(bool bDisable)
+{
+	if (JNIEnv* Env = FAndroidApplication::GetJavaEnv())
+	{
+		// call the java side
+		FJavaWrapper::CallVoidMethod(Env, FJavaWrapper::GameActivityThis, FJavaWrapper::AndroidThunkJava_DisableScreenCapture, bDisable);
+	}
+}
+
 void AndroidThunkCpp_ShowHiddenAlertDialog()
 {
 	if (JNIEnv* Env = FAndroidApplication::GetJavaEnv())
diff --git a/Engine/Source/Runtime/Launch/Public/Android/AndroidJNI.h b/Engine/Source/Runtime/Launch/Public/Android/AndroidJNI.h
--- a/Engine/Source/Runtime/Launch/Public/Android/AndroidJNI.h
+++ b/Engine/Source/Runtime/Launch/Public/Android/AndroidJNI.h
@@ -72,6 +72,10 @@ public:
 	static jmethodID AndroidThunkJava_SetSustainedPerformanceMode;
 	static jmethodID AndroidThunkJava_PushSensorEvents;
 
+	// Screen capture/recording permission
+	static jmethodID AndroidThunkJava_IsScreenCaptureDisabled;
+	static jmethodID AndroidThunkJava_DisableScreenCapture;
+
 	static jmethodID AndroidThunkCpp_VirtualInputIgnoreClick;
 	static jmethodID AndroidThunkCpp_IsVirtuaKeyboardShown;
 	static jmethodID AndroidThunkCpp_IsWebViewShown;
