Fixed: Centering Android splash image only when background is an actual image

diff --git a/Engine/Build/Android/Java/src/com/epicgames/ue4/GameActivity.java.template b/Engine/Build/Android/Java/src/com/epicgames/ue4/GameActivity.java.template
--- a/Engine/Build/Android/Java/src/com/epicgames/ue4/GameActivity.java.template
+++ b/Engine/Build/Android/Java/src/com/epicgames/ue4/GameActivity.java.template
@@ -203,6 +203,7 @@ import android.view.inputmethod.InputConnection;
 import android.view.inputmethod.InputConnectionWrapper;
 import android.util.AttributeSet;
 import android.graphics.Color;
+import android.graphics.drawable.Drawable;
 import android.graphics.drawable.BitmapDrawable;
 import android.widget.TextView;
 import android.widget.TextView.OnEditorActionListener;
@@ -2275,10 +2276,10 @@ public class GameActivity extends $${gameActivitySuperClass}$$ implements Surfac
 					}
 					
 					// Centering background image by default
-					View decorView = mSplashDialog.getWindow().getDecorView();
-					BitmapDrawable SplashBitmap = (BitmapDrawable)decorView.getBackground();
-					if (SplashBitmap != null)
+					Drawable BackgroundDrawable = mSplashDialog.getWindow().getDecorView().getBackground();
+					if (BackgroundDrawable != null && BackgroundDrawable instanceof BitmapDrawable)
 					{
+						BitmapDrawable SplashBitmap = (BitmapDrawable)BackgroundDrawable;
 						SplashBitmap.setGravity(android.view.Gravity.CENTER);
 					}
 					
diff --git a/Engine/Build/Android/Java/src/com/epicgames/ue4/SplashActivity.java b/Engine/Build/Android/Java/src/com/epicgames/ue4/SplashActivity.java
--- a/Engine/Build/Android/Java/src/com/epicgames/ue4/SplashActivity.java
+++ b/Engine/Build/Android/Java/src/com/epicgames/ue4/SplashActivity.java
@@ -21,6 +21,7 @@ import android.provider.Settings;
 import android.view.View;
 import android.view.WindowManager;
 import android.net.Uri;
+import android.graphics.drawable.Drawable;
 import android.graphics.drawable.BitmapDrawable;
 
 import android.support.v4.app.ActivityCompat;
@@ -98,10 +99,10 @@ public class SplashActivity extends Activity
 		}
 
 		// Centering background image by default
-		View decorView = getWindow().getDecorView();
-		BitmapDrawable SplashBitmap = (BitmapDrawable)decorView.getBackground();
-		if (SplashBitmap != null)
+		Drawable BackgroundDrawable = getWindow().getDecorView().getBackground();
+		if (BackgroundDrawable != null && BackgroundDrawable instanceof BitmapDrawable)
 		{
+			BitmapDrawable SplashBitmap = (BitmapDrawable)BackgroundDrawable;
 			SplashBitmap.setGravity(android.view.Gravity.CENTER);
 		}
 		
