Provide a way to pass an obfuscated user id when purchasing using GooglePlay. Google uses this as part of their anti-fraud efforts

diff --git a/Engine/Build/Android/Java/src/com/epicgames/ue4/GameActivity.java.template b/Engine/Build/Android/Java/src/com/epicgames/ue4/GameActivity.java.template
--- a/Engine/Build/Android/Java/src/com/epicgames/ue4/GameActivity.java.template
+++ b/Engine/Build/Android/Java/src/com/epicgames/ue4/GameActivity.java.template
@@ -67,6 +67,9 @@ import java.io.InputStreamReader;
 import java.io.OutputStream;
 import java.io.OutputStreamWriter;
 import java.io.StringWriter;
+import java.io.UnsupportedEncodingException;
+
+import java.security.MessageDigest;
 
 import android.annotation.TargetApi;
 
@@ -4620,13 +4623,35 @@ public class GameActivity extends $${gameActivitySuperClass}$$ implements Surfac
 		}
 	}
 	
-	public boolean AndroidThunkJava_IapBeginPurchase(String ProductId)
+	public boolean AndroidThunkJava_IapBeginPurchase(String ProductId, String AccountId)
 	{
 		Log.debug("[JAVA] - AndroidThunkJava_IapBeginPurchase");
 		boolean bTriggeredPurchase = false;
 		if( IapStoreHelper != null )
 		{
-			bTriggeredPurchase = IapStoreHelper.BeginPurchase(ProductId);
+			// sha-256 the accountId and get the hex string representation
+			String ObfuscatedAccountId = null;
+			if (AccountId != null)
+			{
+				try
+				{
+					MessageDigest md = MessageDigest.getInstance("SHA-256");
+					byte[] sha256hash = md.digest(AccountId.getBytes("UTF-8"));
+					StringBuilder builder = new StringBuilder(sha256hash.length * 2);
+					for (byte b : sha256hash)
+					{
+						builder.append(String.format("%02x", b));
+					}
+					ObfuscatedAccountId = builder.toString();
+				}
+				catch (NoSuchAlgorithmException ae)
+				{
+				}
+				catch (UnsupportedEncodingException ee)
+				{
+				}
+			}
+			bTriggeredPurchase = IapStoreHelper.BeginPurchase(ProductId, ObfuscatedAccountId);
 		}
 		else
 		{
diff --git a/Engine/Build/Android/Java/src/com/epicgames/ue4/GooglePlayStoreHelper.java b/Engine/Build/Android/Java/src/com/epicgames/ue4/GooglePlayStoreHelper.java
--- a/Engine/Build/Android/Java/src/com/epicgames/ue4/GooglePlayStoreHelper.java
+++ b/Engine/Build/Android/Java/src/com/epicgames/ue4/GooglePlayStoreHelper.java
@@ -15,17 +15,53 @@ import android.os.Handler;
 import android.os.IBinder;
 import android.os.Looper;
 import android.os.RemoteException;
-import com.android.vending.billing.util.Base64; //WMM should use different base64 here.
-import com.android.billingclient.api.*;
+import com.android.vending.billing.util.Base64;
+import com.android.vending.billing.IInAppBillingService;
+import com.android.vending.billing.util.Purchase;
 import java.util.ArrayList;
 import java.util.List;
 import org.json.JSONException;
-import org.json.JSONObject; 
+import org.json.JSONObject;
 
 public class GooglePlayStoreHelper implements StoreHelper
 {
+    // Billing response codes
+	// taken from http://developer.android.com/google/play/billing/billing_reference.html Table 1
+    public static final int BILLING_RESPONSE_RESULT_OK = 0;
+    public static final int BILLING_RESPONSE_RESULT_USER_CANCELED = 1;
+	public static final int BILLING_RESPONSE_RESULT_SERVICE_UNAVAILABLE = 2;
+    public static final int BILLING_RESPONSE_RESULT_BILLING_UNAVAILABLE = 3;
+    public static final int BILLING_RESPONSE_RESULT_ITEM_UNAVAILABLE = 4;
+    public static final int BILLING_RESPONSE_RESULT_DEVELOPER_ERROR = 5;
+    public static final int BILLING_RESPONSE_RESULT_ERROR = 6;
+    public static final int BILLING_RESPONSE_RESULT_ITEM_ALREADY_OWNED = 7;
+    public static final int BILLING_RESPONSE_RESULT_ITEM_NOT_OWNED = 8;
+
+	// Keys for the responses from InAppBillingService
+	// taken from http://developer.android.com/google/play/billing/billing_reference.html Table 3
+	public static final String RESPONSE_CODE = "RESPONSE_CODE";
+	public static final String RESPONSE_INAPP_PURCHASE_DATA = "INAPP_PURCHASE_DATA";
+	public static final String RESPONSE_INAPP_SIGNATURE = "INAPP_DATA_SIGNATURE";
+
+	// taken from https://developer.android.com/google/play/billing/billing_integrate.html#purchase
+	public static final String RESPONSE_GET_SKU_DETAILS_LIST = "DETAILS_LIST";
+	public static final String RESPONSE_BUY_INTENT = "BUY_INTENT";
+	public static final String RESPONSE_INAPP_ITEM_LIST = "INAPP_PURCHASE_ITEM_LIST";
+	public static final String RESPONSE_INAPP_PURCHASE_DATA_LIST = "INAPP_PURCHASE_DATA_LIST";
+	public static final String RESPONSE_INAPP_SIGNATURE_LIST = "INAPP_DATA_SIGNATURE_LIST";
+	public static final String INAPP_CONTINUATION_TOKEN = "INAPP_CONTINUATION_TOKEN";
+
+    // Max number of offers to query at once (see getSkuDetails in IInAppbillingServer.aidl)
+	public static final int QueryPurchaseMaxCount = 20;
+
+	// Item types
+	public static final String ITEM_TYPE_INAPP = "inapp";
+
+	// some fields on the getSkuDetails response bundle
+	public static final String GET_SKU_DETAILS_ITEM_LIST = "ITEM_ID_LIST";
+
 	// Our IAB helper interface provided by google.
-	private BillingClient mBillingClient;
+	private IInAppBillingService mService;
 
 	// Flag that determines whether the store is ready to use.
 	private boolean bIsIapSetup;
@@ -41,10 +77,6 @@ public class GooglePlayStoreHelper implements StoreHelper
 
 	private final int UndefinedFailureResponse = -1;
 
-	private int outstandingConsumeRequests = 0;
-
-	private boolean bIsConsumeComplete = false;
-
 	private class GooglePlayProductDescription
 	{
 		// Product offer id 
@@ -59,8 +91,6 @@ public class GooglePlayStoreHelper implements StoreHelper
 		public Float priceRaw;
 		// Local currency code
 		public String currencyCode;
-		// Original Json representation
-		public String originalJson;
 	}
 
 	public interface PurchaseLaunchCallback
@@ -79,47 +109,12 @@ public class GooglePlayStoreHelper implements StoreHelper
 		gameActivity = InGameActivity;
 		productKey = InProductKey;
 
-		PurchasesUpdatedListener purchasesUpdatedListener = new PurchasesUpdatedListener(){
-			@Override
-			public void onPurchasesUpdated(BillingResult billingResult, List<Purchase> purchases){
-				if (billingResult.getResponseCode() == BillingClient.BillingResponseCode.OK
-					&& purchases != null) {
-						for (Purchase purchase : purchases) {
-							onPurchaseResult(billingResult, purchase);
-						}
-				} else
-				{ 
-					Log.debug("[GooglePlayStoreHelper] - GooglePlayStoreHelper::UserCancelled Purchase " + billingResult.getDebugMessage());
-					nativePurchaseComplete(billingResult.getResponseCode(), "", "", "", "");
-					// Handle an error caused by a user cancelling the purchase flow.
-				}
-			}
-		};
-
-		mBillingClient = BillingClient.newBuilder(gameActivity).setListener(purchasesUpdatedListener).enablePendingPurchases().build();
-		
-		mBillingClient.startConnection(new BillingClientStateListener() {
-			@Override
-			public void onBillingSetupFinished(BillingResult billingResult) {
-				if (billingResult.getResponseCode() == BillingClient.BillingResponseCode.OK) {
-					bIsIapSetup = true;
-					Log.debug("In-app billing supported for " + gameActivity.getPackageName());
-				}
-				else
-				{
-					Log.debug("In-app billing NOT supported for " + gameActivity.getPackageName() + " error " + billingResult.getResponseCode());
-				}
-			}
-    
-			@Override
-			public void onBillingServiceDisconnected() {
-				// Try to restart the connection on the next request to
-				// Google Play by calling the startConnection() method.
-				bIsIapSetup = false;
-			}
-		});
+		Intent serviceIntent = new Intent("com.android.vending.billing.InAppBillingService.BIND");
+		serviceIntent.setPackage("com.android.vending");
+		gameActivity.bindService(serviceIntent, mServiceConn, Context.BIND_AUTO_CREATE);
 	}
 
+
 	///////////////////////////////////////////////////////
 	// The StoreHelper interfaces implementation for Google Play Store.
 
@@ -149,161 +144,195 @@ public class GooglePlayStoreHelper implements StoreHelper
 				skuList.add(productId);
 			}
 
-			Log.debug("[GooglePlayStoreHelper] - NumSkus: " + skuList.size());
+			ArrayList<String> productIds = new ArrayList<String>();
+			ArrayList<String> titles = new ArrayList<String>();
+			ArrayList<String> descriptions = new ArrayList<String>();
+			ArrayList<String> prices = new ArrayList<String>();
+			ArrayList<Float> pricesRaw = new ArrayList<Float>();
+			ArrayList<String> currencyCodes = new ArrayList<String>();
 
-			SkuDetailsParams.Builder skuDetailsParams = SkuDetailsParams.newBuilder();
-			skuDetailsParams = skuDetailsParams.setSkusList(skuList).setType(BillingClient.SkuType.INAPP);
-			
-			mBillingClient.querySkuDetailsAsync(skuDetailsParams.build(),
-				new SkuDetailsResponseListener() {
-					@Override
-					public void onSkuDetailsResponse(BillingResult billingResult,
-						List<SkuDetails> skuDetailsList) 
+			while (skuList.size() > 0)
+			{
+				Log.debug("[GooglePlayStoreHelper] - NumSkus: " + skuList.size());
+				int numSkus = Math.min(QueryPurchaseMaxCount, skuList.size());
+				Log.debug("[GooglePlayStoreHelper] - Grabbing from 0 to " + numSkus);
+				List<String> skuSubList = skuList.subList(0, numSkus);
+				ArrayList<String> realSubList = new ArrayList<String>(skuSubList);
+				ArrayList<GooglePlayProductDescription> productDescriptions = new ArrayList<GooglePlayProductDescription>();
+
+				int response = QueryInAppPurchasesInternal(realSubList, productDescriptions);
+				if (response == BILLING_RESPONSE_RESULT_OK)
+				{
+					for (GooglePlayProductDescription product : productDescriptions)
 					{
-						// Process the result.
-						ArrayList<GooglePlayProductDescription> productDescriptions = new ArrayList<GooglePlayProductDescription>();
-						int response = billingResult.getResponseCode();
-						Log.debug("[GooglePlayStoreHelper] - GooglePlayStoreHelper::QueryInAppPurchases - Response " + response + " SkuDetails:" + skuDetailsList.toString());
-
-						// Should we send JSON, or somehow serialize this so we can have random access fields?
-						ArrayList<String> productIds = new ArrayList<String>();
-						ArrayList<String> titles = new ArrayList<String>();
-						ArrayList<String> descriptions = new ArrayList<String>();
-						ArrayList<String> prices = new ArrayList<String>();
-						ArrayList<Float> pricesRaw = new ArrayList<Float>();
-						ArrayList<String> currencyCodes = new ArrayList<String>();
-						ArrayList<String> originalJson = new ArrayList<String>();
-
-						if(response == BillingClient.BillingResponseCode.OK)
-						{
-							for (SkuDetails thisSku : skuDetailsList)
-							{
-								GooglePlayProductDescription newDescription = new GooglePlayProductDescription();
-								newDescription.id = thisSku.getSku();
-								newDescription.title = thisSku.getTitle();
-								newDescription.description = thisSku.getDescription();
-								newDescription.price = thisSku.getPrice();
-								double priceRaw = thisSku.getPriceAmountMicros() / 1000000.0;
-								newDescription.priceRaw = (float)priceRaw;
-								newDescription.currencyCode = thisSku.getPriceCurrencyCode();
-								newDescription.originalJson = thisSku.getOriginalJson();
-								productDescriptions.add(newDescription);
-							}
-						
-							for (GooglePlayProductDescription product : productDescriptions)
-							{
-								productIds.add(product.id);
-								Log.debug("[GooglePlayStoreHelper] - GooglePlayStoreHelper::QueryInAppPurchases - Parsing details for: " + product.id);
+						productIds.add(product.id);
+						Log.debug("[GooglePlayStoreHelper] - GooglePlayStoreHelper::QueryInAppPurchases - Parsing details for: " + product.id);
 
-								titles.add(product.title);
-								Log.debug("[GooglePlayStoreHelper] - title: " + product.title);
+						titles.add(product.title);
+						Log.debug("[GooglePlayStoreHelper] - title: " + product.title);
 
-								descriptions.add(product.description);
-								Log.debug("[GooglePlayStoreHelper] - description: " + product.description);
+						descriptions.add(product.description);
+						Log.debug("[GooglePlayStoreHelper] - description: " + product.description);
 
-								prices.add(product.price);
-								Log.debug("[GooglePlayStoreHelper] - price: " + product.price);
+						prices.add(product.price);
+						Log.debug("[GooglePlayStoreHelper] - price: " + product.price);
 
-								pricesRaw.add(product.priceRaw);
-								Log.debug("[GooglePlayStoreHelper] - price_amount_micros: " + product.priceRaw);
+						pricesRaw.add(product.priceRaw);
+						Log.debug("[GooglePlayStoreHelper] - price_amount_micros: " + product.priceRaw);
 
-								currencyCodes.add(product.currencyCode);
-								Log.debug("[GooglePlayStoreHelper] - price_currency_code: " + product.currencyCode);
+						currencyCodes.add(product.currencyCode);
+						Log.debug("[GooglePlayStoreHelper] - price_currency_code: " + product.currencyCode);
+					}
+				}
+				else
+				{
+					Log.debug("[GooglePlayStoreHelper] - GooglePlayStoreHelper::QueryInAppPurchases - Failed with: " + response);
+					// If anything fails right now, stop immediately, not sure how to reconcile partial success/failure
+					nativeQueryComplete(response, null, null, null, null, null, null);
+					return false;
+				}
 
-								originalJson.add(product.originalJson);
-								Log.debug("[GooglePlayStoreHelper] - original_json: " + product.originalJson);
+				skuSubList.clear();
+			}
 
-							}
-						}
-						else
-						{
-							Log.debug("[GooglePlayStoreHelper] - GooglePlayStoreHelper::QueryInAppPurchases - Failed with: " + response);
-							// If anything fails right now, stop immediately, not sure how to reconcile partial success/failure
-							nativeQueryComplete(response, null, null, null, null, null, null, null);
-						}
-						
-						float[] pricesRawPrimitive = new float[pricesRaw.size()];
-						for (int i = 0; i < pricesRaw.size(); i++)
-						{
-							pricesRawPrimitive[i] = pricesRaw.get(i);
-						}
+			float[] pricesRawPrimitive = new float[pricesRaw.size()];
+			for (int i = 0; i < pricesRaw.size(); i++)
+			{
+				pricesRawPrimitive[i] = pricesRaw.get(i);
+			}
 
-						Log.debug("[GooglePlayStoreHelper] - GooglePlayStoreHelper::QueryInAppPurchases " + productIds.size() + " items - Success!");
-						nativeQueryComplete(BillingClient.BillingResponseCode.OK, productIds.toArray(new String[productIds.size()]), titles.toArray(new String[titles.size()]), descriptions.toArray(new String[descriptions.size()]), prices.toArray(new String[prices.size()]), pricesRawPrimitive, currencyCodes.toArray(new String[currencyCodes.size()]), originalJson.toArray(new String[originalJson.size()]));
-						Log.debug("[GooglePlayStoreHelper] - nativeQueryComplete done!");
-					}
-				}
-			);			
+			Log.debug("[GooglePlayStoreHelper] - GooglePlayStoreHelper::QueryInAppPurchases " + productIds.size() + " items - Success!");
+			nativeQueryComplete(BILLING_RESPONSE_RESULT_OK, productIds.toArray(new String[productIds.size()]), titles.toArray(new String[titles.size()]), descriptions.toArray(new String[descriptions.size()]), prices.toArray(new String[prices.size()]), pricesRawPrimitive, currencyCodes.toArray(new String[currencyCodes.size()]));
+			Log.debug("[GooglePlayStoreHelper] - nativeQueryComplete done!");
 		}
 		else
 		{
 			// nothing to query
 			Log.debug("[GooglePlayStoreHelper] - no products given to query");
-			nativeQueryComplete(UndefinedFailureResponse, null, null, null, null, null, null, null);
-			return false;
+			nativeQueryComplete(UndefinedFailureResponse, null, null, null, null, null, null);
 		}
 
 		return true;
 	}
 
+	public int QueryInAppPurchasesInternal(ArrayList<String> skuList, ArrayList<GooglePlayProductDescription> outProducts)
+	{
+		Log.debug("[GooglePlayStoreHelper] - GooglePlayStoreHelper::QueryInAppPurchasesInternal");
+
+		Bundle querySkus = new Bundle();
+		querySkus.putStringArrayList(GET_SKU_DETAILS_ITEM_LIST, skuList);
+
+		int response = UndefinedFailureResponse;
+		try
+		{
+			Bundle skuDetails = mService.getSkuDetails(3, gameActivity.getPackageName(), ITEM_TYPE_INAPP, querySkus);
+
+			response = skuDetails.getInt(RESPONSE_CODE);
+			Log.debug("[GooglePlayStoreHelper] - GooglePlayStoreHelper::QueryInAppPurchases - Response " + response + " Bundle:" + skuDetails.toString());
+			if (response == BILLING_RESPONSE_RESULT_OK)
+			{
+				ArrayList<String> responseList = skuDetails.getStringArrayList(RESPONSE_GET_SKU_DETAILS_LIST);
+				for (String thisResponse : responseList)
+				{
+					JSONObject object = new JSONObject(thisResponse);
+
+					GooglePlayProductDescription newDescription = new GooglePlayProductDescription();
+					newDescription.id = object.getString("productId");
+					newDescription.title = object.getString("title");
+					newDescription.description = object.getString("description");
+					newDescription.price = object.getString("price");
+					double priceRaw = object.getDouble("price_amount_micros") / 1000000.0;
+					newDescription.priceRaw = (float)priceRaw;
+					newDescription.currencyCode = object.getString("price_currency_code");
+					outProducts.add(newDescription);
+				}
+			}
+		}
+		catch(Exception e)
+		{
+			Log.debug("[GooglePlayStoreHelper] - GooglePlayStoreHelper::QueryInAppPurchasesInternal - Failed! " + e.getMessage());
+		}
+
+		return response;
+	}
+
 	/**
 	 * Start the purchase flow for a particular sku
 	 */
 	final int purchaseIntentIdentifier = 1001;
-	public boolean BeginPurchase(String ProductID)
+	public boolean BeginPurchase(String ProductID, String ObfuscatedAccountId)
 	{
-		ArrayList<String> skuList = new ArrayList<String>(1);
-		skuList.add(ProductID);
-
-		SkuDetailsParams.Builder skuDetailsParams = SkuDetailsParams.newBuilder();
-		skuDetailsParams.setSkusList(skuList).setType(BillingClient.SkuType.INAPP);
-		mBillingClient.querySkuDetailsAsync(skuDetailsParams.build(),
-				new SkuDetailsResponseListener() {
-					@Override
-					public void onSkuDetailsResponse(BillingResult billingResult,
-						List<SkuDetails> skuDetailsList) 
-					{
-						
-						int response = billingResult.getResponseCode();
-						if(response == BillingClient.BillingResponseCode.OK)
-						{
-							for(SkuDetails skuDetails : skuDetailsList)
-							{
-								BillingFlowParams.Builder flowParams = BillingFlowParams.newBuilder().setSkuDetails(skuDetails);
-			
-			
-								if (gameActivity.IsInVRMode())
-								{
-									int responseCode = mBillingClient.isFeatureSupported(BillingClient.FeatureType.IN_APP_ITEMS_ON_VR).getResponseCode();
-									boolean isSupportedInVR = (responseCode == BillingClient.BillingResponseCode.OK);
-									if (isSupportedInVR)
-									{
-										Log.debug("[GooglePlayStoreHelper] - GooglePlayStoreHelper::BeginPurchase - v7 VR purchase" + skuDetails.getSku());
-										flowParams = flowParams.setVrPurchaseFlow(true);
-									}
-									else
-									{
-										Log.debug("[GooglePlayStoreHelper] - GooglePlayStoreHelper::BeginPurchase - v3 IAB purchase:" + skuDetails.getSku());
-										flowParams = flowParams.setVrPurchaseFlow(false);
-									}
-								}
-								else
-								{
-									Log.debug("[GooglePlayStoreHelper] - GooglePlayStoreHelper::BeginPurchase - v3 IAB purchase:" + skuDetails.getSku());
-									flowParams = flowParams.setVrPurchaseFlow(false);
-								}
+		try
+		{
+			String devPayload = GenerateDevPayload(ProductID);
+			Bundle buyIntentBundle = null;
+			int response = -1;
 
-								int responseCode = mBillingClient.launchBillingFlow(gameActivity, flowParams.build()).getResponseCode();
-								Log.debug("[GooglePlayStoreHelper] - GooglePlayStoreHelper::BeginPurchase - Launching billing flow " + skuDetails.getSku());
-							}
-						}
-						else
-						{
-							Log.debug("[GooglePlayStoreHelper] - GooglePlayStoreHelper::BeginPurchase - Failed! " + TranslateServerResponseCode(response));
-							nativePurchaseComplete(UndefinedFailureResponse, "", "", "", "");
-						}
+			Bundle bundle = new Bundle();
+			if (ObfuscatedAccountId != null)
+			{
+				bundle.putString("accountId", ObfuscatedAccountId);
+			}
+			if (gameActivity.IsInVRMode())
+			{
+				bundle.putBoolean("vr", true);
+			}
+
+			response = mService.isBillingSupportedExtraParams(7, gameActivity.getPackageName(), ITEM_TYPE_INAPP, bundle);
+			if (response == BILLING_RESPONSE_RESULT_OK)
+			{
+				if (gameActivity.IsInVRMode())
+				{
+					Log.debug("[GooglePlayStoreHelper] - GooglePlayStoreHelper::BeginPurchase - v7 VR purchase" + ProductID);
+				}
+				else
+				{
+					Log.debug("[GooglePlayStoreHelper] - GooglePlayStoreHelper::BeginPurchase - v7 IAB purchase" + ProductID);
+				}
+
+				buyIntentBundle = mService.getBuyIntentExtraParams(7, gameActivity.getPackageName(), ProductID, ITEM_TYPE_INAPP, devPayload, bundle);
+			}
+			else
+			{
+				Log.debug("[GooglePlayStoreHelper] - GooglePlayStoreHelper::BeginPurchase - v3 IAB purchase:" + ProductID);
+				buyIntentBundle = mService.getBuyIntent(3, gameActivity.getPackageName(), ProductID, ITEM_TYPE_INAPP, devPayload);
+			}
+
+			response = buyIntentBundle.getInt(RESPONSE_CODE);
+			if (response == BILLING_RESPONSE_RESULT_OK)
+			{
+				Log.debug("[GooglePlayStoreHelper] - GooglePlayStoreHelper::BeginPurchase - Starting Intent to buy " + ProductID);
+				PendingIntent pendingIntent = buyIntentBundle.getParcelable(RESPONSE_BUY_INTENT);
+				if (pendingIntent != null)
+				{
+					PurchaseLaunchCallback callback =  gameActivity.getPurchaseLaunchCallback();
+					if (callback != null)
+					{
+						callback.launchForResult(pendingIntent, purchaseIntentIdentifier);
 					}
-		});
+					else
+					{
+						gameActivity.startIntentSenderForResult(pendingIntent.getIntentSender(), purchaseIntentIdentifier, new Intent(), Integer.valueOf(0), Integer.valueOf(0), Integer.valueOf(0));
+					}
+				}
+				else
+				{
+					Log.debug("[GooglePlayStoreHelper] - GooglePlayStoreHelper::BeginPurchase - pendingIntent was null, possible non consumed item");
+					nativePurchaseComplete(UndefinedFailureResponse, ProductID, "", "", "");
+				}
+			}
+			else
+			{
+				Log.debug("[GooglePlayStoreHelper] - GooglePlayStoreHelper::BeginPurchase - Failed with error: " + TranslateServerResponseCode(response));
+				nativePurchaseComplete(response, ProductID, "", "", "");
+			}
+		}
+		catch(Exception e)
+		{
+			Log.debug("[GooglePlayStoreHelper] - GooglePlayStoreHelper::BeginPurchase - Failed! " + e.getMessage());
+			nativePurchaseComplete(UndefinedFailureResponse, ProductID, "", "", "");
+		}
+
 		return true;
 	}
 
@@ -313,113 +342,99 @@ public class GooglePlayStoreHelper implements StoreHelper
 	public boolean RestorePurchases(String[] InProductIDs, boolean[] bConsumable)
 	{
 		Log.debug("[GooglePlayStoreHelper] - GooglePlayStoreHelper::RestorePurchases");
+		ArrayList<String> ownedSkus = new ArrayList<String>();
+		ArrayList<String> purchaseDataList = new ArrayList<String>();
+		ArrayList<String> signatureList = new ArrayList<String>();
 
-		ArrayList<Purchase> ownedProducts = new ArrayList<Purchase>();
-		int responseCode = GatherOwnedPurchaseData(ownedProducts);
-		if (responseCode == BillingClient.BillingResponseCode.OK)
+		// On first pass the continuation token should be null.
+		// This will allow us to gather large sets of purchased items recursively
+		int responseCode = GatherOwnedPurchaseData(ownedSkus, purchaseDataList, signatureList, null);
+		if (responseCode == BILLING_RESPONSE_RESULT_OK)
 		{
-			Log.debug("[GooglePlayStoreHelper] - GooglePlayStoreHelper::RestorePurchases - User has previously purchased " + ownedProducts.size() + " inapp products" );
-
-			final ArrayList<Purchase> f_ownedProducts = ownedProducts;
+			Log.debug("[GooglePlayStoreHelper] - GooglePlayStoreHelper::RestorePurchases - User has previously purchased " + ownedSkus.size() + " inapp products" );
 
-			final ArrayList<String> f_ownedSkus = new ArrayList<String>();
-			final ArrayList<String> f_purchaseDataList = new ArrayList<String>();
-			final ArrayList<String> f_signatureList = new ArrayList<String>();
+			final ArrayList<String> f_ownedSkus = ownedSkus;
+			final ArrayList<String> f_purchaseDataList = purchaseDataList;
+			final ArrayList<String> f_signatureList = signatureList;
 			final String[] RestoreProductIDs = InProductIDs;
 			final boolean[] bRestoreConsumableFlags = bConsumable;
 
-			// Is this here to ensure callbacks are fired on the main thread?
 			Handler mainHandler = new Handler(Looper.getMainLooper());
 			mainHandler.post(new Runnable()
 			{
 				@Override
 				public void run()
 				{
-					Log.debug("[GooglePlayStoreHelper] - GooglePlayStoreHelper::RestorePurchases - Now consuming any purchases that may have been missed.");
-					final ArrayList<String> productTokens = new ArrayList<String>();
-					final ArrayList<String> receipts = new ArrayList<String>();
-					final ArrayList<Integer> cachedErrorResponses = new ArrayList<Integer>();
-
-					bIsConsumeComplete = false;
-					// @todo possible bug, restores ALL ownedSkus, not just those provided by caller in RestoreProductIDs
-					for (Purchase ownedProduct : f_ownedProducts)
+					try
 					{
-						
-						productTokens.add(ownedProduct.getPurchaseToken());
-
-						boolean bTryToConsume = false;
-
-						// This is assuming that all purchases should be consumed. Consuming a purchase that is meant to be a one-time purchase makes it so the
-						// user is able to buy it again. Also, it makes it so the purchase will not be able to be restored again in the future.
+						Log.debug("[GooglePlayStoreHelper] - GooglePlayStoreHelper::RestorePurchases - Now consuming any purchases that may have been missed.");
+						ArrayList<String> productTokens = new ArrayList<String>();
+						ArrayList<String> receipts = new ArrayList<String>();
+						int cachedResponse = 0;
 
-						for (int Idy = 0; Idy < RestoreProductIDs.length; Idy++)
+						// @todo possible bug, restores ALL ownedSkus, not just those provided by caller in RestoreProductIDs
+						for (int Idx = 0; Idx < f_ownedSkus.size(); Idx++)
 						{
-							if( ownedProduct.getSku().equals(RestoreProductIDs[Idy]) )
+							String purchaseData = f_purchaseDataList.get(Idx);
+							String dataSignature = f_signatureList.get(Idx);
+
+							try
 							{
-								if(Idy < bRestoreConsumableFlags.length)
-								{
-									bTryToConsume = bRestoreConsumableFlags[Idy];
-									Log.debug("[GooglePlayStoreHelper] - GooglePlayStoreHelper::RestorePurchases - Found Consumable Flag for Product " + ownedProduct.getSku() + " bConsumable = " + bTryToConsume);
-								}
-								break;
-							}
-						}
+								Purchase purchase = new Purchase(ITEM_TYPE_INAPP, purchaseData, dataSignature);
+								productTokens.add(purchase.getToken());
 
-						final Purchase constProduct = ownedProduct;
-						if (bTryToConsume)
-						{
-							Log.debug("[GooglePlayStoreHelper] - GooglePlayStoreHelper::RestorePurchases - Attempting to consume " + ownedProduct.getSku());
-							ConsumeParams consumeParams = ConsumeParams.newBuilder().setPurchaseToken(ownedProduct.getPurchaseToken()).build();
-							mBillingClient.consumeAsync(consumeParams, (new ConsumeResponseListener() {
-								@Override
-								public void onConsumeResponse(BillingResult billingResult, String purchaseToken) {
-									if(billingResult.getResponseCode() == BillingClient.BillingResponseCode.OK)
-									{
-										// How do we get purchase here.. need to figure out a way to persist cachedResponse and receipts
-										Log.debug("[GooglePlayStoreHelper] - GooglePlayStoreHelper::RestorePurchases - Purchase restored for " + constProduct.getSku());
-										String receipt = Base64.encode(constProduct.getOriginalJson().getBytes());
-										receipts.add(receipt);
-										
-										f_ownedSkus.add(constProduct.getSku());
-										f_signatureList.add(constProduct.getSignature());
-									}
-									else
-									{
-										Log.debug("[GooglePlayStoreHelper] - GooglePlayStoreHelper::RestorePurchases - consumePurchase failed for " + constProduct.getSku() + " with error " + billingResult.getResponseCode());
-										receipts.add("");
-										cachedErrorResponses.add(billingResult.getResponseCode());
-									}
-									outstandingConsumeRequests--;
+								boolean bTryToConsume = false;
+								int consumeResponse = 0;
 
-									if(outstandingConsumeRequests <= 0)
+								// This is assuming that all purchases should be consumed. Consuming a purchase that is meant to be a one-time purchase makes it so the
+								// user is able to buy it again. Also, it makes it so the purchase will not be able to be restored again in the future.
+
+								for (int Idy = 0; Idy < RestoreProductIDs.length; Idy++)
+								{
+									if( purchase.getSku().equals(RestoreProductIDs[Idy]) )
 									{
-										Log.debug("[GooglePlayStoreHelper] - GooglePlayStoreHelper::RestorePurchases - Success!");
-										int finalResponse = (cachedErrorResponses.size() > 0) ? cachedErrorResponses.get(0) : BillingClient.BillingResponseCode.OK;
-										nativeRestorePurchasesComplete(finalResponse, f_ownedSkus.toArray(new String[f_ownedSkus.size()]), productTokens.toArray(new String[productTokens.size()]), receipts.toArray(new String[receipts.size()]), f_signatureList.toArray(new String[f_signatureList.size()]));
-										bIsConsumeComplete = true;
+										if(Idy < bRestoreConsumableFlags.length)
+										{
+											bTryToConsume = bRestoreConsumableFlags[Idy];
+											Log.debug("[GooglePlayStoreHelper] - GooglePlayStoreHelper::RestorePurchases - Found Consumable Flag for Product " + purchase.getSku() + " bConsumable = " + bTryToConsume);
+										}
+										break;
 									}
 								}
-							}));
-								
-							outstandingConsumeRequests++;
-						}
-						else
-						{
-							// How do we get purchase here.. need to figure out a way to persist cachedResponse and receipts
-							Log.debug("[GooglePlayStoreHelper] - GooglePlayStoreHelper::RestorePurchases - Purchase restored for " + constProduct.getSku());
-							String receipt = Base64.encode(constProduct.getOriginalJson().getBytes());
-							receipts.add(receipt);
-										
-							f_ownedSkus.add(constProduct.getSku());
-							f_signatureList.add(constProduct.getSignature());
+
+								if (bTryToConsume)
+								{
+									Log.debug("[GooglePlayStoreHelper] - GooglePlayStoreHelper::RestorePurchases - Attempting to consume " + purchase.getSku());
+									consumeResponse = mService.consumePurchase(3, gameActivity.getPackageName(), purchase.getToken());
+								}
+
+								if (consumeResponse == 0)
+								{
+									Log.debug("[GooglePlayStoreHelper] - GooglePlayStoreHelper::RestorePurchases - Purchase restored for " + purchase.getSku());
+									String receipt = Base64.encode(purchase.getOriginalJson().getBytes());
+									receipts.add(receipt);
+								}
+								else
+								{
+									Log.debug("[GooglePlayStoreHelper] - GooglePlayStoreHelper::RestorePurchases - consumePurchase failed for " + purchase.getSku() + " with error " + consumeResponse);
+									receipts.add("");
+									cachedResponse = consumeResponse;
+								}
+							}
+							catch (JSONException e)
+							{
+								Log.debug("[GooglePlayStoreHelper] - GooglePlayStoreHelper::RestorePurchases - Failed to parse receipt! " + e.getMessage());
+								productTokens.add("");
+								receipts.add("");
+							}
 						}
+						Log.debug("[GooglePlayStoreHelper] - GooglePlayStoreHelper::RestorePurchases - Success!");
+						nativeRestorePurchasesComplete(cachedResponse, f_ownedSkus.toArray(new String[f_ownedSkus.size()]), productTokens.toArray(new String[productTokens.size()]), receipts.toArray(new String[receipts.size()]), f_signatureList.toArray(new String[f_signatureList.size()]));
 					}
-
-					// There was nothing to consume, so we are done!
-					if(outstandingConsumeRequests <= 0 && !bIsConsumeComplete)
+					catch (Exception e)
 					{
-						Log.debug("[GooglePlayStoreHelper] - GooglePlayStoreHelper::RestorePurchases - Success!");
-						nativeRestorePurchasesComplete(BillingClient.BillingResponseCode.OK, f_ownedSkus.toArray(new String[f_ownedSkus.size()]), productTokens.toArray(new String[productTokens.size()]), receipts.toArray(new String[receipts.size()]), f_signatureList.toArray(new String[f_signatureList.size()]));
+						Log.debug("[GooglePlayStoreHelper] - GooglePlayStoreHelper::RestorePurchases - consumePurchase failed. " + e.getMessage());
+						nativeRestorePurchasesComplete(UndefinedFailureResponse, null, null, null, null);
 					}
 				}
 			});
@@ -441,23 +456,35 @@ public class GooglePlayStoreHelper implements StoreHelper
 		ArrayList<String> purchaseDataList = new ArrayList<String>();
 		ArrayList<String> signatureList = new ArrayList<String>();
 
-		ArrayList<Purchase> ownedProducts = new ArrayList<Purchase>();
-
-		int responseCode = GatherOwnedPurchaseData(ownedProducts);
-		if (responseCode == BillingClient.BillingResponseCode.OK)
+		// On first pass the continuation token should be null.
+		// This will allow us to gather large sets of purchased items recursively
+		int responseCode = GatherOwnedPurchaseData(ownedSkus, purchaseDataList, signatureList, null);
+		if (responseCode == BILLING_RESPONSE_RESULT_OK)
 		{
 			Log.debug("[GooglePlayStoreHelper] - GooglePlayStoreHelper::QueryExistingPurchases - User has previously purchased " + ownedSkus.size() + " inapp products" );
 
 			ArrayList<String> productTokens = new ArrayList<String>();
 			ArrayList<String> receipts = new ArrayList<String>();
 
-			for (Purchase purchase : ownedProducts)
+			for (int Idx = 0; Idx < ownedSkus.size(); Idx++)
 			{
-				productTokens.add(purchase.getPurchaseToken());
-				ownedSkus.add(purchase.getSku());
-				signatureList.add(purchase.getSignature());
-				String receipt = Base64.encode(purchase.getOriginalJson().getBytes());
-				receipts.add(receipt);
+				String purchaseData = purchaseDataList.get(Idx);
+				String dataSignature = signatureList.get(Idx);
+
+				try
+				{
+					Purchase purchase = new Purchase(ITEM_TYPE_INAPP, purchaseData, dataSignature);
+					productTokens.add(purchase.getToken());
+
+					String receipt = Base64.encode(purchase.getOriginalJson().getBytes());
+					receipts.add(receipt);
+				}
+				catch (JSONException e)
+				{
+					Log.debug("[GooglePlayStoreHelper] - GooglePlayStoreHelper::QueryExistingPurchases - Failed to parse receipt! " + e.getMessage());
+					productTokens.add("");
+					receipts.add("");
+				}
 			}
 
 			Log.debug("[GooglePlayStoreHelper] - GooglePlayStoreHelper::QueryExistingPurchases - Success!");
@@ -484,24 +511,71 @@ public class GooglePlayStoreHelper implements StoreHelper
 			@Override
 			public void run()
 			{
-				ConsumeParams consumeParams = ConsumeParams.newBuilder().setPurchaseToken(f_purchaseToken).build();
-				mBillingClient.consumeAsync(consumeParams, (new ConsumeResponseListener() {
-					@Override
-					public void onConsumeResponse(BillingResult billingResult, String purchaseToken) {
-						if(billingResult.getResponseCode() == BillingClient.BillingResponseCode.OK)
-						{
-							Log.debug("[GooglePlayStoreHelper] - ConsumePurchase success");
-						}
-						else
-						{
-							Log.debug("[GooglePlayStoreHelper] - ConsumePurchase failed with error " + TranslateServerResponseCode(billingResult.getResponseCode()));
-						}
+				try
+				{
+					Log.debug("[GooglePlayStoreHelper] - Consuming token: " + f_purchaseToken);
+					int consumeResponse = mService.consumePurchase(3, gameActivity.getPackageName(), f_purchaseToken);
+					if (consumeResponse == BILLING_RESPONSE_RESULT_OK)
+					{
+						Log.debug("[GooglePlayStoreHelper] - ConsumePurchase success");
 					}
-				}));
+					else
+					{
+						Log.debug("[GooglePlayStoreHelper] - ConsumePurchase failed with error " + TranslateServerResponseCode(consumeResponse));
+					}
+				}
+				catch(Exception e)
+				{
+					Log.debug("[GooglePlayStoreHelper] - ConsumePurchase failed. " + e.getMessage());
+				}
 			}
 		});
 	}
 
+	///////////////////////////////////////////////////////
+	// IInAppBillingService service registration and unregistration
+
+	private ServiceConnection mServiceConn = new ServiceConnection()
+	{
+		@Override
+		public void onServiceDisconnected(ComponentName name)
+		{
+			Log.debug("[GooglePlayStoreHelper] - ServiceConnection::onServiceDisconnected");
+			mService = null;
+			bIsIapSetup = false;
+		}
+
+		@Override
+		public void onServiceConnected(ComponentName name, IBinder service)
+		{
+			Log.debug("[GooglePlayStoreHelper] - ServiceConnection::onServiceConnected");
+			mService = IInAppBillingService.Stub.asInterface(service);
+			bIsIapSetup = true;
+
+			try
+			{
+                Log.debug("Checking for in-app billing 3 support.");
+
+                // check for in-app billing v3 support
+                int response = mService.isBillingSupported(3, gameActivity.getPackageName(), ITEM_TYPE_INAPP);
+                if (response != BILLING_RESPONSE_RESULT_OK)
+				{
+					Log.debug("In-app billing version 3 NOT supported for " + gameActivity.getPackageName() + " error " + response);
+                }
+				else
+				{
+					Log.debug("In-app billing version 3 supported for " + gameActivity.getPackageName());
+				}
+			}
+            catch (RemoteException e)
+			{
+				e.printStackTrace();
+				return;
+            }
+		}
+	};
+
+
 	///////////////////////////////////////////////////////
 	// Game Activity/Context driven methods we need to listen for.
 
@@ -512,9 +586,9 @@ public class GooglePlayStoreHelper implements StoreHelper
 	{
 		Log.debug("[GooglePlayStoreHelper] - GooglePlayStoreHelper::onDestroy");
 
-		if (mBillingClient != null)
+		if (mService != null)
 		{
-			mBillingClient = null; //WMM ???
+			gameActivity.unbindService(mServiceConn);
 		}
 	}
 
@@ -522,97 +596,141 @@ public class GooglePlayStoreHelper implements StoreHelper
 	 * Route taken by the Purchase workflow. We listen for our purchaseIntentIdentifier request code and
 	 * handle the response accordingly
 	 */
-	public boolean onPurchaseResult(BillingResult billingResult, Purchase purchase)
+	public boolean onActivityResult(int requestCode, int resultCode, Intent data)
 	{
-		Log.debug("[GooglePlayStoreHelper] - GooglePlayStoreHelper::onPurchaseResult");
+		Log.debug("[GooglePlayStoreHelper] - GooglePlayStoreHelper::onActivityResult");
 
-		if (billingResult == null)
+		if (requestCode == purchaseIntentIdentifier)
 		{
-			Log.debug("Null data in purchase activity result.");
-			nativePurchaseComplete(UndefinedFailureResponse, "", "", "", "");
-			return true;
-		}
-
-		int responseCode = billingResult.getResponseCode();
-		if (responseCode == BillingClient.BillingResponseCode.OK)
-		{
-			Log.debug("[GooglePlayStoreHelper] - GooglePlayStoreHelper::onActivityResult - Processing purchase result. Response Code: " + TranslateServerResponseCode(responseCode));
-			Log.debug("Purchase data: " + purchase.toString());
-			Log.debug("Data signature: " + purchase.getSignature());
+		    if (data == null)
+			{
+				Log.debug("Null data in purchase activity result.");
+				nativePurchaseComplete(UndefinedFailureResponse, "", "", "", "");
+				return true;
+			}
 
-			if(purchase.getPurchaseState() == Purchase.PurchaseState.PURCHASED)
+			int responseCode = data.getIntExtra(RESPONSE_CODE, 0);
+			if (resultCode == Activity.RESULT_OK)
 			{
-				final String sku = purchase.getSku();
-				final Purchase f_purchase = purchase;
-				Handler mainHandler = new Handler(Looper.getMainLooper());
-				mainHandler.post(new Runnable()
+				final String purchaseData = data.getStringExtra(RESPONSE_INAPP_PURCHASE_DATA);
+				final String dataSignature = data.getStringExtra(RESPONSE_INAPP_SIGNATURE);
+
+				Log.debug("[GooglePlayStoreHelper] - GooglePlayStoreHelper::onActivityResult - Processing purchase result. Response Code: " + TranslateServerResponseCode(responseCode));
+				if (responseCode == BILLING_RESPONSE_RESULT_OK)
 				{
-					@Override
-					public void run()
+					Log.debug("Purchase data: " + purchaseData);
+					Log.debug("Data signature: " + dataSignature);
+					Log.debug("Extras: " + data.getExtras());
+
+					if (purchaseData == null || dataSignature == null)
+					{
+						Log.debug("[GooglePlayStoreHelper] - GooglePlayStoreHelper::onActivityResult - Either purchaseData or dataSignature is null");
+						nativePurchaseComplete(UndefinedFailureResponse, "", "", "", "");
+					}
+
+					try
 					{
-						String receipt = Base64.encode(f_purchase.getOriginalJson().getBytes());
-						nativePurchaseComplete(BillingClient.BillingResponseCode.OK, sku, f_purchase.getPurchaseToken(), receipt, f_purchase.getSignature());
+						final Purchase purchase = new Purchase(ITEM_TYPE_INAPP, purchaseData, dataSignature);
+						final String sku = purchase.getSku();
+						String developerPayload = purchase.getDeveloperPayload();
+						if (VerifyPayload(developerPayload, sku))
+						{
+							Handler mainHandler = new Handler(Looper.getMainLooper());
+							mainHandler.post(new Runnable()
+							{
+								@Override
+								public void run()
+								{
+									String receipt = Base64.encode(purchase.getOriginalJson().getBytes());
+									nativePurchaseComplete(BILLING_RESPONSE_RESULT_OK, sku, purchase.getToken(), receipt, purchase.getSignature());
+								}
+							});
+						}
+						else
+						{
+							Log.debug("[GooglePlayStoreHelper] - GooglePlayStoreHelper::onActivityResult - Failed for verify developer payload for " + sku);
+							nativePurchaseComplete(UndefinedFailureResponse, sku, "", "", "");
+						}
+					}
+					catch (JSONException e)
+					{
+						Log.debug("[GooglePlayStoreHelper] - GooglePlayStoreHelper::onActivityResult - Failed for purchase request! " + e.getMessage());
+						nativePurchaseComplete(UndefinedFailureResponse, "", "", "", "");
 					}
-				});
+				}
+				else
+				{
+					Log.debug("[GooglePlayStoreHelper] - GooglePlayStoreHelper::onActivityResult - Failed for purchase request!. " + TranslateServerResponseCode(responseCode));
+					nativePurchaseComplete(responseCode, "", "", "", "");
+				}
+			}
+			else if (resultCode == Activity.RESULT_CANCELED)
+			{
+				Log.debug("[GooglePlayStoreHelper] - GooglePlayStoreHelper::onActivityResult - Purchase canceled." + TranslateServerResponseCode(responseCode));
+				nativePurchaseComplete(BILLING_RESPONSE_RESULT_USER_CANCELED, "", "", "", "");
 			}
 			else
 			{
-				// This is a pending purchase.. need to hold off on reporting as purchase to game until it actually occurs.
-				// Special case - if pending purchase is USER_CANCELLED but the SKU is returned, then it is deferred
-				final String sku = purchase.getSku();
-				nativePurchaseComplete(BillingClient.BillingResponseCode.USER_CANCELED, sku, "", "", "");
+				Log.debug("[GooglePlayStoreHelper] - GooglePlayStoreHelper::onActivityResult - Purchase failed. Result code: " +
+					Integer.toString(resultCode) + ". Response: " + TranslateServerResponseCode(responseCode));
+				nativePurchaseComplete(UndefinedFailureResponse, "", "", "", "");
 			}
-		}
-		else if (responseCode == BillingClient.BillingResponseCode.USER_CANCELED)
-		{
-			Log.debug("[GooglePlayStoreHelper] - GooglePlayStoreHelper::onActivityResult - Purchase canceled." + TranslateServerResponseCode(responseCode));
-			nativePurchaseComplete(BillingClient.BillingResponseCode.USER_CANCELED, "", "", "", "");
-		}
-		else
-		{
-			Log.debug("[GooglePlayStoreHelper] - GooglePlayStoreHelper::onActivityResult - Purchase failed. Result code: " +
-				Integer.toString(responseCode) + ". Response: " + TranslateServerResponseCode(responseCode));
-			nativePurchaseComplete(UndefinedFailureResponse, "", "", "", "");
+
+			return true;
 		}
 
-		return true;
+		return false;
 	}
 
 
 	///////////////////////////////////////////////////////
 	// Internal helper functions that deal assist with various IAB related events
 
+	/**
+	 * Create a UE4 specific unique string that will be used to verify purchases are legit.
+	 */
+	private String GenerateDevPayload( String ProductId )
+	{
+		return "ue4." + ProductId;
+	}
+
+	/**
+	 * Check the returned payload matches one for the product we are buying.
+	 */
+	private boolean VerifyPayload( String ExistingPayload, String ProductId )
+	{
+		String GeneratedPayload = GenerateDevPayload(ProductId);
+		Log.debug("[GooglePlayStoreHelper] - GooglePlayStoreHelper::VerifyPayload, ExistingPayload: " + ExistingPayload);
+		Log.debug("[GooglePlayStoreHelper] - GooglePlayStoreHelper::VerifyPayload, GeneratedPayload: " + GeneratedPayload);
+
+		return ExistingPayload.equals(GeneratedPayload);
+	}
+
 	/**
 	 * Get a text tranlation of the Response Codes returned by google play.
 	 */
 	private String TranslateServerResponseCode(final int serverResponseCode)
 	{
-		// taken from https://developer.android.com/reference/com/android/billingclient/api/BillingClient.BillingResponseCode
+		// taken from http://developer.android.com/google/play/billing/billing_reference.html
 		switch(serverResponseCode)
 		{
-			case BillingClient.BillingResponseCode.OK:
+			case BILLING_RESPONSE_RESULT_OK:
 				return "Success";
-			case BillingClient.BillingResponseCode.USER_CANCELED:
+			case BILLING_RESPONSE_RESULT_USER_CANCELED:
 				return "User pressed back or canceled a dialog";
-			case BillingClient.BillingResponseCode.SERVICE_UNAVAILABLE:
+			case BILLING_RESPONSE_RESULT_SERVICE_UNAVAILABLE:
 				return "Network connection is down";
-			case BillingClient.BillingResponseCode.SERVICE_TIMEOUT:
-				return "The request has reached the maximum timeout before Google Play responds";
-			case BillingClient.BillingResponseCode.SERVICE_DISCONNECTED:
-				return "Play Store service is not connected now - potentially transient state";
-			case BillingClient.BillingResponseCode.BILLING_UNAVAILABLE:
+			case BILLING_RESPONSE_RESULT_BILLING_UNAVAILABLE:
 				return "Billing API version is not supported for the type requested";
-			case BillingClient.BillingResponseCode.FEATURE_NOT_SUPPORTED:
-				return "Requested feature is not supported by Play Store on the current device";
-			case BillingClient.BillingResponseCode.ITEM_UNAVAILABLE:
+			case BILLING_RESPONSE_RESULT_ITEM_UNAVAILABLE:
 				return "Requested product is not available for purchase";
-			case BillingClient.BillingResponseCode.DEVELOPER_ERROR:
+			case BILLING_RESPONSE_RESULT_DEVELOPER_ERROR:
 				return "Invalid arguments provided to the API. This error can also indicate that the application was not correctly signed or properly set up for In-app Billing in Google Play, or does not have the necessary permissions in its manifest";
-			case BillingClient.BillingResponseCode.ERROR:
+			case BILLING_RESPONSE_RESULT_ERROR:
 				return "Fatal error during the API action";
-			case BillingClient.BillingResponseCode.ITEM_ALREADY_OWNED:
+			case BILLING_RESPONSE_RESULT_ITEM_ALREADY_OWNED:
 				return "Failure to purchase since item is already owned";
-			case BillingClient.BillingResponseCode.ITEM_NOT_OWNED:
+			case BILLING_RESPONSE_RESULT_ITEM_NOT_OWNED:
 				return "Failure to consume since item is not owned";
 			default:
 				return "Unknown Server Response Code";
@@ -623,25 +741,39 @@ public class GooglePlayStoreHelper implements StoreHelper
 	 * Recursive functionality to gather all of the purchases owned by a user.
 	 * if the user owns a lot of products then we need to getPurchases again with a continuationToken
 	 */
-	private int GatherOwnedPurchaseData(ArrayList<Purchase> inPurchases)
+	private int GatherOwnedPurchaseData(ArrayList<String> inOwnedSkus, ArrayList<String> inPurchaseDataList, ArrayList<String> inSignatureList, String inContinuationToken)
 	{
 		int responseCode = UndefinedFailureResponse;
 		try
 		{
-			Purchase.PurchasesResult purchasesResult = mBillingClient.queryPurchases(BillingClient.SkuType.INAPP);
+			Bundle ownedItems = mService.getPurchases(3, gameActivity.getPackageName(), ITEM_TYPE_INAPP, inContinuationToken);
 
-			responseCode = purchasesResult.getResponseCode();
+			responseCode = ownedItems.getInt(RESPONSE_CODE);
 			Log.debug("[GooglePlayStoreHelper] - GooglePlayStoreHelper::GatherOwnedPurchaseData - getPurchases result. Response Code: " + TranslateServerResponseCode(responseCode));
-			if (responseCode == BillingClient.BillingResponseCode.OK)
+			if (responseCode == BILLING_RESPONSE_RESULT_OK)
 			{
-				List<Purchase> ownedItems = purchasesResult.getPurchasesList();
-				for (Purchase purchase : ownedItems)
+				ArrayList<String> ownedSkus = ownedItems.getStringArrayList(RESPONSE_INAPP_ITEM_LIST);
+				ArrayList<String> purchaseDataList = ownedItems.getStringArrayList(RESPONSE_INAPP_PURCHASE_DATA_LIST);
+				ArrayList<String> signatureList = ownedItems.getStringArrayList(RESPONSE_INAPP_SIGNATURE_LIST);
+				String continuationToken = ownedItems.getString(INAPP_CONTINUATION_TOKEN);
+
+				for (int Idx = 0; Idx < purchaseDataList.size(); Idx++)
 				{
-					// Don't Mark pending purchases as owned
-					if(purchase.getPurchaseState() == Purchase.PurchaseState.PURCHASED)
-					{
-						inPurchases.add(purchase);
-					}
+					String sku = ownedSkus.get(Idx);
+					inOwnedSkus.add(sku);
+
+					String purchaseData = purchaseDataList.get(Idx);
+					inPurchaseDataList.add(purchaseData);
+
+					String signature = signatureList.get(Idx);
+					inSignatureList.add(signature);
+				}
+
+				// if continuationToken != null, call getPurchases again
+				// and pass in the token to retrieve more items
+				if (continuationToken != null)
+				{
+					return GatherOwnedPurchaseData(inOwnedSkus, inPurchaseDataList, inSignatureList, inContinuationToken);
 				}
 			}
 		}
@@ -654,7 +786,7 @@ public class GooglePlayStoreHelper implements StoreHelper
 	}
 
 	// Callback that notify the C++ implementation that a task has completed
-	public native void nativeQueryComplete(int responseCode, String[] productIDs, String[] titles, String[] descriptions, String[] prices, float[] pricesRaw, String[] currencyCodes, String[] originalJson );
+	public native void nativeQueryComplete(int responseCode, String[] productIDs, String[] titles, String[] descriptions, String[] prices, float[] pricesRaw, String[] currencyCodes );
 	public native void nativePurchaseComplete(int responseCode, String ProductId, String ProductToken, String ReceiptData, String Signature);
 	public native void nativeRestorePurchasesComplete(int responseCode, String[] ProductIds, String[] ProductTokens, String[] ReceiptsData, String[] Signatures);
 	public native void nativeQueryExistingPurchasesComplete(int responseCode, String[] ProductIds, String[] ProductTokens, String[] ReceiptsData, String[] Signatures);
diff --git a/Engine/Build/Android/Java/src/com/epicgames/ue4/StoreHelper.java b/Engine/Build/Android/Java/src/com/epicgames/ue4/StoreHelper.java
--- a/Engine/Build/Android/Java/src/com/epicgames/ue4/StoreHelper.java
+++ b/Engine/Build/Android/Java/src/com/epicgames/ue4/StoreHelper.java
@@ -8,10 +8,11 @@ import android.os.Bundle;
 public interface StoreHelper
 {
 	public boolean QueryInAppPurchases(String[] ProductIDs);
-	public boolean BeginPurchase(String ProductID);
+	public boolean BeginPurchase(String ProductID, String ObfuscatedAccountId);
 	public boolean IsAllowedToMakePurchases();
 	public void ConsumePurchase(String purchaseToken);
 	public boolean QueryExistingPurchases();
 	public boolean RestorePurchases(String[] InProductIDs, boolean[] bConsumable);
 	public void	onDestroy();
+	public boolean onActivityResult(int requestCode, int resultCode, Intent data);
 }
\ No newline at end of file
diff --git a/Engine/Plugins/Online/Android/OnlineSubsystemGameCircle/Build/Android/Java/src/com/epicgames/ue4/AmazonStoreHelper.java b/Engine/Plugins/Online/Android/OnlineSubsystemGameCircle/Build/Android/Java/src/com/epicgames/ue4/AmazonStoreHelper.java
--- a/Engine/Plugins/Online/Android/OnlineSubsystemGameCircle/Build/Android/Java/src/com/epicgames/ue4/AmazonStoreHelper.java
+++ b/Engine/Plugins/Online/Android/OnlineSubsystemGameCircle/Build/Android/Java/src/com/epicgames/ue4/AmazonStoreHelper.java
@@ -107,7 +107,7 @@ public class AmazonStoreHelper implements StoreHelper, PurchasingListener
 		return true;
 	}
 
-	public boolean BeginPurchase(String ProductID)
+	public boolean BeginPurchase(String ProductID, string ObfuscatedAccountId)
 	{
 		Log.debug("AmazonStoreHelper.BeginPurchase - " + ProductID);
 
diff --git a/Engine/Plugins/Online/Android/OnlineSubsystemGooglePlay/Source/Private/OnlinePurchaseGooglePlay.cpp b/Engine/Plugins/Online/Android/OnlineSubsystemGooglePlay/Source/Private/OnlinePurchaseGooglePlay.cpp
--- a/Engine/Plugins/Online/Android/OnlineSubsystemGooglePlay/Source/Private/OnlinePurchaseGooglePlay.cpp
+++ b/Engine/Plugins/Online/Android/OnlineSubsystemGooglePlay/Source/Private/OnlinePurchaseGooglePlay.cpp
@@ -94,8 +94,13 @@ void FOnlinePurchaseGooglePlay::Checkout(const FUniqueNetId& UserId, const FPurc
 				{
 					const FPurchaseCheckoutRequest::FPurchaseOfferEntry& Offer = CheckoutRequest.PurchaseOffers[0];
 
-					extern bool AndroidThunkCpp_Iap_BeginPurchase(const FString&);
-					bStarted = AndroidThunkCpp_Iap_BeginPurchase(Offer.OfferId);
+					extern bool AndroidThunkCpp_Iap_BeginPurchase(const FString&, const FString&);
+					FString AccountId;
+					if (CheckoutRequest.AccountId.IsSet())
+					{
+						AccountId = CheckoutRequest.AccountId.GetValue();
+					}
+					bStarted = AndroidThunkCpp_Iap_BeginPurchase(Offer.OfferId, AccountId);
 					UE_LOG_ONLINE_PURCHASE(Display, TEXT("Created Transaction? - %s"),
 						bStarted ? TEXT("Created a transaction.") : TEXT("Failed to create a transaction."));
 
diff --git a/Engine/Plugins/Online/Android/OnlineSubsystemGooglePlay/Source/Private/OnlineStoreInterfaceGooglePlay.cpp b/Engine/Plugins/Online/Android/OnlineSubsystemGooglePlay/Source/Private/OnlineStoreInterfaceGooglePlay.cpp
--- a/Engine/Plugins/Online/Android/OnlineSubsystemGooglePlay/Source/Private/OnlineStoreInterfaceGooglePlay.cpp
+++ b/Engine/Plugins/Online/Android/OnlineSubsystemGooglePlay/Source/Private/OnlineStoreInterfaceGooglePlay.cpp
@@ -9,9 +9,7 @@
 #include "OnlineSubsystemGooglePlay.h"
 #include <jni.h>
 #include "Android/AndroidJavaEnv.h"
-#include "Internationalization/Internationalization.h"
-#include "Internationalization/Culture.h"
-#include "Internationalization/FastDecimalFormat.h"
+
 
 ////////////////////////////////////////////////////////////////////
 /// FOnlineStoreGooglePlay implementation
@@ -81,44 +79,86 @@ bool FOnlineStoreGooglePlay::QueryForAvailablePurchases(const TArray<FString>& P
 	return true;
 }
 
-TSharedRef<FInAppPurchaseProductInfo> ConvertStoreOfferToProduct(const FOnlineStoreOffer& Offer)
+
+JNI_METHOD void Java_com_epicgames_ue4_GooglePlayStoreHelper_nativeQueryComplete(JNIEnv* jenv, jobject thiz, jsize responseCode, jobjectArray productIDs, jobjectArray titles, jobjectArray descriptions, jobjectArray prices, jfloatArray pricesRaw, jobjectArray currencyCodes)
 {
-	TSharedRef<FInAppPurchaseProductInfo> NewProductInfo = MakeShareable(new FInAppPurchaseProductInfo());
+	TArray<FInAppPurchaseProductInfo> ProvidedProductInformation;
+	EGooglePlayBillingResponseCode EGPResponse = (EGooglePlayBillingResponseCode)responseCode;
+	bool bWasSuccessful = (EGPResponse == EGooglePlayBillingResponseCode::Ok);
+
+	if (jenv && bWasSuccessful)
+	{
+		jsize NumProducts = jenv->GetArrayLength(productIDs);
+		jsize NumTitles = jenv->GetArrayLength(titles);
+		jsize NumDescriptions = jenv->GetArrayLength(descriptions);
+		jsize NumPrices = jenv->GetArrayLength(prices);
+		jsize NumPricesRaw = jenv->GetArrayLength(pricesRaw);
+		jsize NumCurrencyCodes = jenv->GetArrayLength(currencyCodes);
 
-	NewProductInfo->Identifier = Offer.OfferId;
-	NewProductInfo->DisplayName = Offer.Title.ToString();
+		ensure((NumProducts == NumTitles) && (NumProducts == NumDescriptions) && (NumProducts == NumPrices) && (NumProducts == NumPricesRaw) && (NumProducts == NumCurrencyCodes));
 
-	NewProductInfo->DisplayDescription = Offer.Description.ToString(); // Google has only one description, map it to (short) description to match iOS
-	NewProductInfo->DisplayPrice = Offer.PriceText.ToString();
-	NewProductInfo->CurrencyCode = Offer.CurrencyCode;
+		jfloat *PricesRaw = jenv->GetFloatArrayElements(pricesRaw, 0);
 
-	// Convert the backend stated price into its base units
-	FInternationalization& I18N = FInternationalization::Get();
-	const FCulture& Culture = *I18N.GetCurrentCulture();
+		for (jsize Idx = 0; Idx < NumProducts; Idx++)
+		{
+			// Build the product information strings.
+
+			FInAppPurchaseProductInfo NewProductInfo;
+
+			NewProductInfo.Identifier = FJavaHelper::FStringFromLocalRef(jenv, (jstring)jenv->GetObjectArrayElement(productIDs, Idx));
+			NewProductInfo.DisplayName = FJavaHelper::FStringFromLocalRef(jenv, (jstring)jenv->GetObjectArrayElement(titles, Idx));
+			NewProductInfo.DisplayDescription = FJavaHelper::FStringFromLocalRef(jenv, (jstring)jenv->GetObjectArrayElement(descriptions, Idx));
+			NewProductInfo.DisplayPrice = FJavaHelper::FStringFromLocalRef(jenv, (jstring)jenv->GetObjectArrayElement(prices, Idx));
+			
+			NewProductInfo.RawPrice = PricesRaw[Idx];
+			
+			NewProductInfo.CurrencyCode = FJavaHelper::FStringFromLocalRef(jenv, (jstring)jenv->GetObjectArrayElement(currencyCodes, Idx));
+
+			ProvidedProductInformation.Add(NewProductInfo);
+
+			FPlatformMisc::LowLevelOutputDebugStringf(TEXT("\nProduct Identifier: %s, Name: %s, Description: %s, Price: %s, Price Raw: %f, Currency Code: %s\n"),
+				*NewProductInfo.Identifier,
+				*NewProductInfo.DisplayName,
+				*NewProductInfo.DisplayDescription,
+				*NewProductInfo.DisplayPrice,
+				NewProductInfo.RawPrice,
+				*NewProductInfo.CurrencyCode);
+		}
+	}
 
-	const FDecimalNumberFormattingRules& FormattingRules = Culture.GetCurrencyFormattingRules(NewProductInfo->CurrencyCode);
-	const FNumberFormattingOptions& FormattingOptions = FormattingRules.CultureDefaultFormattingOptions;
-	double Val = static_cast<double>(Offer.NumericPrice) / static_cast<double>(FMath::Pow(10.0f, FormattingOptions.MaximumFractionalDigits));
+	DECLARE_CYCLE_STAT(TEXT("FSimpleDelegateGraphTask.ProcessQueryIapResult"), STAT_FSimpleDelegateGraphTask_ProcessQueryIapResult, STATGROUP_TaskGraphTasks);
 
-	NewProductInfo->RawPrice = Val;
+	FPlatformMisc::LowLevelOutputDebugStringf(TEXT("Adding task Success: %d Response: %s"), bWasSuccessful, ToString(EGPResponse));
 
-	return NewProductInfo;
+	FSimpleDelegateGraphTask::CreateAndDispatchWhenReady(
+		FSimpleDelegateGraphTask::FDelegate::CreateLambda([=]()
+	{
+		if (IOnlineSubsystem* const OnlineSub = IOnlineSubsystem::Get(GOOGLEPLAY_SUBSYSTEM))
+		{
+			FOnlineSubsystemGooglePlay* const OnlineSubGP = static_cast<FOnlineSubsystemGooglePlay* const>(OnlineSub);
+			if (OnlineSubGP)
+			{
+				FPlatformMisc::LowLevelOutputDebugStringf(TEXT("TriggerOnGooglePlayAvailableIAPQueryCompleteDelegates %s Size: %d"), ToString(EGPResponse), ProvidedProductInformation.Num());
+				OnlineSubGP->TriggerOnGooglePlayAvailableIAPQueryCompleteDelegates(EGPResponse, ProvidedProductInformation);
+			}
+		}
+		FPlatformMisc::LowLevelOutputDebugStringf(TEXT("In-App Purchase query was completed  %s\n"), bWasSuccessful ? TEXT("successfully") : TEXT("unsuccessfully"));
+	}),
+		GET_STATID(STAT_FSimpleDelegateGraphTask_ProcessQueryIapResult),
+		nullptr,
+		ENamedThreads::GameThread
+		);
 }
 
-void FOnlineStoreGooglePlay::OnGooglePlayAvailableIAPQueryComplete(EGooglePlayBillingResponseCode InResponseCode, const TArray<FOnlineStoreOffer>& AvailablePurchases)
+
+void FOnlineStoreGooglePlay::OnGooglePlayAvailableIAPQueryComplete(EGooglePlayBillingResponseCode InResponseCode, const TArray<FInAppPurchaseProductInfo>& AvailablePurchases)
 {
 	UE_LOG_ONLINE_STORE(Display, TEXT("FOnlineStoreGooglePlay::OnGooglePlayAvailableIAPQueryComplete"));
 
-	TArray<FInAppPurchaseProductInfo> TempInfos;
-	for (const FOnlineStoreOffer& AvailablePurchase : AvailablePurchases)
-	{
-		TempInfos.Add(*ConvertStoreOfferToProduct(AvailablePurchase));
-	}
-
 	if (ReadObject.IsValid())
 	{
 		ReadObject->ReadState = (InResponseCode == EGooglePlayBillingResponseCode::Ok) ? EOnlineAsyncTaskState::Done : EOnlineAsyncTaskState::Failed;
-		ReadObject->ProvidedProductInformation.Insert(TempInfos, 0);
+		ReadObject->ProvidedProductInformation.Insert(AvailablePurchases, 0);
 	}
 
 	CurrentQueryTask->ProcessQueryAvailablePurchasesResults(InResponseCode == EGooglePlayBillingResponseCode::Ok);
@@ -138,8 +178,8 @@ bool FOnlineStoreGooglePlay::BeginPurchase(const FInAppPurchaseProductRequest& P
 		CachedPurchaseStateObject = InPurchaseStateObject;
 		CachedPurchaseStateObject->bIsConsumable = ProductRequest.bIsConsumable;
 
-		extern bool AndroidThunkCpp_Iap_BeginPurchase(const FString&);
-		bCreatedNewTransaction = AndroidThunkCpp_Iap_BeginPurchase(ProductRequest.ProductIdentifier);
+		extern bool AndroidThunkCpp_Iap_BeginPurchase(const FString&, const FString&);
+		bCreatedNewTransaction = AndroidThunkCpp_Iap_BeginPurchase(ProductRequest.ProductIdentifier, FString());
 		UE_LOG_ONLINE_STORE(Display, TEXT("Created Transaction? - %s"), 
 			bCreatedNewTransaction ? TEXT("Created a transaction.") : TEXT("Failed to create a transaction."));
 
@@ -165,6 +205,48 @@ bool FOnlineStoreGooglePlay::BeginPurchase(const FInAppPurchaseProductRequest& P
 	return bCreatedNewTransaction;
 }
 
+
+JNI_METHOD void Java_com_epicgames_ue4_GooglePlayStoreHelper_nativePurchaseComplete(JNIEnv* jenv, jobject thiz, jsize responseCode, jstring productId, jstring productToken, jstring receiptData, jstring signature)
+{
+	FString ProductId, ProductToken, ReceiptData, Signature;
+	EGooglePlayBillingResponseCode EGPResponse = (EGooglePlayBillingResponseCode)responseCode;
+
+	bool bWasSuccessful = (EGPResponse == EGooglePlayBillingResponseCode::Ok);
+	if (bWasSuccessful)
+	{
+		ProductId = FJavaHelper::FStringFromParam(jenv, productId);
+		ProductToken = FJavaHelper::FStringFromParam(jenv, productToken);
+		ReceiptData = FJavaHelper::FStringFromParam(jenv, receiptData);
+		Signature = FJavaHelper::FStringFromParam(jenv, signature);
+	}
+
+	FGoogleTransactionData TransactionData(ProductId, ProductToken, ReceiptData, Signature);
+
+	FPlatformMisc::LowLevelOutputDebugStringf(TEXT("1... Response: %s, Transaction %s"), ToString(EGPResponse), *TransactionData.ToDebugString());
+
+	DECLARE_CYCLE_STAT(TEXT("FSimpleDelegateGraphTask.ProcessIapResult"), STAT_FSimpleDelegateGraphTask_ProcessIapResult, STATGROUP_TaskGraphTasks);
+
+	FSimpleDelegateGraphTask::CreateAndDispatchWhenReady(
+		FSimpleDelegateGraphTask::FDelegate::CreateLambda([=]()
+	{
+		FPlatformMisc::LowLevelOutputDebugStringf(TEXT("In-App Purchase was completed  %s\n"), bWasSuccessful ? TEXT("successfully") : TEXT("unsuccessfully"));
+		if (IOnlineSubsystem* const OnlineSub = IOnlineSubsystem::Get(GOOGLEPLAY_SUBSYSTEM))
+		{
+			FOnlineSubsystemGooglePlay* const OnlineSubGP = static_cast<FOnlineSubsystemGooglePlay* const>(OnlineSub);
+			if (OnlineSubGP)
+			{
+				FPlatformMisc::LowLevelOutputDebugStringf(TEXT("2... Response %s Transaction %s"), ToString(EGPResponse), *TransactionData.ToDebugString());
+				OnlineSubGP->TriggerOnGooglePlayProcessPurchaseCompleteDelegates(EGPResponse, TransactionData);
+			}
+		}
+	}),
+		GET_STATID(STAT_FSimpleDelegateGraphTask_ProcessIapResult),
+		nullptr,
+		ENamedThreads::GameThread
+		);
+}
+
+
 void FOnlineStoreGooglePlay::OnProcessPurchaseResult(EGooglePlayBillingResponseCode InResponseCode, const FGoogleTransactionData& InTransactionData)
 {
 	UE_LOG_ONLINE_STORE(Display, TEXT("FOnlineStoreGooglePlay::OnProcessPurchaseResult"));
@@ -226,6 +308,61 @@ bool FOnlineStoreGooglePlay::RestorePurchases(const TArray<FInAppPurchaseProduct
 	return bSentAQueryRequest;
 }
 
+JNI_METHOD void Java_com_epicgames_ue4_GooglePlayStoreHelper_nativeRestorePurchasesComplete(JNIEnv* jenv, jobject thiz, jsize responseCode, jobjectArray ProductIDs, jobjectArray ProductTokens, jobjectArray ReceiptsData, jobjectArray Signatures)
+{
+	TArray<FGoogleTransactionData> RestoredPurchaseInfo;
+
+	EGooglePlayBillingResponseCode EGPResponse = (EGooglePlayBillingResponseCode)responseCode;
+	bool bWasSuccessful = (EGPResponse == EGooglePlayBillingResponseCode::Ok);
+
+	if (jenv && bWasSuccessful)
+	{
+		jsize NumProducts = jenv->GetArrayLength(ProductIDs);
+		jsize NumProductTokens = jenv->GetArrayLength(ProductTokens);
+		jsize NumReceipts = jenv->GetArrayLength(ReceiptsData);
+		jsize NumSignatures = jenv->GetArrayLength(Signatures);
+
+		ensure((NumProducts == NumProductTokens) && (NumProducts == NumReceipts) && (NumProducts == NumSignatures));
+
+		for (jsize Idx = 0; Idx < NumProducts; Idx++)
+		{
+			// Build the restore product information strings.
+			FInAppPurchaseRestoreInfo RestoreInfo;
+
+			const auto OfferId = FJavaHelper::FStringFromLocalRef(jenv, (jstring)jenv->GetObjectArrayElement(ProductIDs, Idx));
+			const auto ProductToken = FJavaHelper::FStringFromLocalRef(jenv, (jstring)jenv->GetObjectArrayElement(ProductTokens, Idx));
+			const auto ReceiptData = FJavaHelper::FStringFromLocalRef(jenv, (jstring)jenv->GetObjectArrayElement(ReceiptsData, Idx));
+			const auto SignatureData = FJavaHelper::FStringFromLocalRef(jenv, (jstring)jenv->GetObjectArrayElement(Signatures, Idx));
+			
+			FGoogleTransactionData RestoredPurchase(OfferId, ProductToken, ReceiptData, SignatureData);
+			RestoredPurchaseInfo.Add(RestoredPurchase);
+
+			FPlatformMisc::LowLevelOutputDebugStringf(TEXT("Restored Transaction: %s"), *RestoredPurchase.ToDebugString());
+		}
+	}
+
+	DECLARE_CYCLE_STAT(TEXT("FSimpleDelegateGraphTask.RestorePurchases"), STAT_FSimpleDelegateGraphTask_RestorePurchases, STATGROUP_TaskGraphTasks);
+
+	FSimpleDelegateGraphTask::CreateAndDispatchWhenReady(
+		FSimpleDelegateGraphTask::FDelegate::CreateLambda([=]()
+	{
+		FPlatformMisc::LowLevelOutputDebugStringf(TEXT("Restoring In-App Purchases was completed  %s\n"), bWasSuccessful ? TEXT("successfully") : TEXT("unsuccessfully"));
+
+		if (IOnlineSubsystem* const OnlineSub = IOnlineSubsystem::Get(GOOGLEPLAY_SUBSYSTEM))
+		{
+			FOnlineSubsystemGooglePlay* const OnlineSubGP = static_cast<FOnlineSubsystemGooglePlay* const>(OnlineSub);
+			if (OnlineSubGP)
+			{
+				OnlineSubGP->TriggerOnGooglePlayRestorePurchasesCompleteDelegates(EGPResponse, RestoredPurchaseInfo);
+			}
+		}
+	}),
+		GET_STATID(STAT_FSimpleDelegateGraphTask_RestorePurchases),
+		nullptr,
+		ENamedThreads::GameThread
+		);
+}
+
 void FOnlineStoreGooglePlay::OnRestorePurchasesComplete(EGooglePlayBillingResponseCode InResponseCode, const TArray<FGoogleTransactionData>& InRestoredPurchases)
 {
 	UE_LOG_ONLINE_STORE(Verbose, TEXT("FOnlineStoreGooglePlay::OnRestorePurchasesComplete Response: %s Num: %d"), ToString(InResponseCode), InRestoredPurchases.Num());
@@ -253,3 +390,54 @@ void FOnlineStoreGooglePlay::OnRestorePurchasesComplete(EGooglePlayBillingRespon
 	TriggerOnInAppPurchaseRestoreCompleteDelegates(IAPState);
 }
 
+JNI_METHOD void Java_com_epicgames_ue4_GooglePlayStoreHelper_nativeQueryExistingPurchasesComplete(JNIEnv* jenv, jobject thiz, jsize responseCode, jobjectArray ProductIDs, jobjectArray ProductTokens, jobjectArray ReceiptsData, jobjectArray Signatures)
+{
+	TArray<FGoogleTransactionData> ExistingPurchaseInfo;
+
+	EGooglePlayBillingResponseCode EGPResponse = (EGooglePlayBillingResponseCode)responseCode;
+
+	bool bWasSuccessful = (EGPResponse == EGooglePlayBillingResponseCode::Ok);
+	if (jenv && bWasSuccessful)
+	{
+		jsize NumProducts = jenv->GetArrayLength(ProductIDs);
+		jsize NumProductTokens = jenv->GetArrayLength(ProductTokens);
+		jsize NumReceipts = jenv->GetArrayLength(ReceiptsData);
+		jsize NumSignatures = jenv->GetArrayLength(Signatures);
+
+		ensure((NumProducts == NumProductTokens) && (NumProducts == NumReceipts) && (NumProducts == NumSignatures));
+
+		for (jsize Idx = 0; Idx < NumProducts; Idx++)
+		{
+			// Build the product information strings.
+			const auto OfferId = FJavaHelper::FStringFromLocalRef(jenv, (jstring)jenv->GetObjectArrayElement(ProductIDs, Idx));
+			const auto ProductToken = FJavaHelper::FStringFromLocalRef(jenv, (jstring)jenv->GetObjectArrayElement(ProductTokens, Idx));
+			const auto ReceiptData = FJavaHelper::FStringFromLocalRef(jenv, (jstring)jenv->GetObjectArrayElement(ReceiptsData, Idx));
+			const auto SignatureData = FJavaHelper::FStringFromLocalRef(jenv, (jstring)jenv->GetObjectArrayElement(Signatures, Idx));
+			
+			FGoogleTransactionData ExistingPurchase(OfferId, ProductToken, ReceiptData, SignatureData);
+			ExistingPurchaseInfo.Add(ExistingPurchase);
+
+			FPlatformMisc::LowLevelOutputDebugStringf(TEXT("\Existing Product Identifier: %s"), *ExistingPurchase.ToDebugString());
+		}
+	}
+
+	DECLARE_CYCLE_STAT(TEXT("FSimpleDelegateGraphTask.QueryExistingPurchases"), STAT_FSimpleDelegateGraphTask_QueryExistingPurchases, STATGROUP_TaskGraphTasks);
+
+	FSimpleDelegateGraphTask::CreateAndDispatchWhenReady(
+		FSimpleDelegateGraphTask::FDelegate::CreateLambda([=]()
+	{
+		FPlatformMisc::LowLevelOutputDebugStringf(TEXT("Query existing purchases was completed %s\n"), bWasSuccessful ? TEXT("successfully") : TEXT("unsuccessfully"));
+		if (IOnlineSubsystem* const OnlineSub = IOnlineSubsystem::Get(GOOGLEPLAY_SUBSYSTEM))
+		{
+			FOnlineSubsystemGooglePlay* const OnlineSubGP = static_cast<FOnlineSubsystemGooglePlay* const>(OnlineSub);
+			if (OnlineSubGP)
+			{
+				OnlineSubGP->TriggerOnGooglePlayQueryExistingPurchasesCompleteDelegates(EGPResponse, ExistingPurchaseInfo);
+			}
+		}
+	}),
+		GET_STATID(STAT_FSimpleDelegateGraphTask_QueryExistingPurchases),
+		nullptr,
+		ENamedThreads::GameThread
+		);
+}
diff --git a/Engine/Plugins/Online/OnlineSubsystem/Source/Public/Interfaces/OnlinePurchaseInterface.h b/Engine/Plugins/Online/OnlineSubsystem/Source/Public/Interfaces/OnlinePurchaseInterface.h
--- a/Engine/Plugins/Online/OnlineSubsystem/Source/Public/Interfaces/OnlinePurchaseInterface.h
+++ b/Engine/Plugins/Online/OnlineSubsystem/Source/Public/Interfaces/OnlinePurchaseInterface.h
@@ -62,6 +62,9 @@ public:
 	};	
 	/** List of offers being purchased */
 	TArray<FPurchaseOfferEntry> PurchaseOffers;
+
+	/** Optional AccountId to include with the purchase request (if supported) */
+	TOptional<FString> AccountId;
 };
 
 /**
diff --git a/Engine/Source/Runtime/Launch/Private/Android/AndroidJNI.cpp b/Engine/Source/Runtime/Launch/Private/Android/AndroidJNI.cpp
--- a/Engine/Source/Runtime/Launch/Private/Android/AndroidJNI.cpp
+++ b/Engine/Source/Runtime/Launch/Private/Android/AndroidJNI.cpp
@@ -202,7 +202,7 @@ void FJavaWrapper::FindGooglePlayBillingMethods(JNIEnv* Env)
 	JavaStringClass = FindClassGlobalRef(Env, "java/lang/String", false);
 	AndroidThunkJava_IapSetupService = FindMethod(Env, GoogleServicesClassID, "AndroidThunkJava_IapSetupService", "(Ljava/lang/String;)V", bIsStoreOptional);
 	AndroidThunkJava_IapQueryInAppPurchases = FindMethod(Env, GoogleServicesClassID, "AndroidThunkJava_IapQueryInAppPurchases", "([Ljava/lang/String;)Z", bIsStoreOptional);
-	AndroidThunkJava_IapBeginPurchase = FindMethod(Env, GoogleServicesClassID, "AndroidThunkJava_IapBeginPurchase", "(Ljava/lang/String;)Z", bIsStoreOptional);
+	AndroidThunkJava_IapBeginPurchase = FindMethod(Env, GoogleServicesClassID, "AndroidThunkJava_IapBeginPurchase", "(Ljava/lang/String;Ljava/lang/String;)Z", bIsStoreOptional);
 	AndroidThunkJava_IapIsAllowedToMakePurchases = FindMethod(Env, GoogleServicesClassID, "AndroidThunkJava_IapIsAllowedToMakePurchases", "()Z", bIsStoreOptional);
 	AndroidThunkJava_IapRestorePurchases = FindMethod(Env, GoogleServicesClassID, "AndroidThunkJava_IapRestorePurchases", "([Ljava/lang/String;[Z)Z", bIsStoreOptional);
 	AndroidThunkJava_IapConsumePurchase = FindMethod(Env, GoogleServicesClassID, "AndroidThunkJava_IapConsumePurchase", "(Ljava/lang/String;)Z", bIsStoreOptional);
@@ -1338,7 +1338,7 @@ bool AndroidThunkCpp_Iap_QueryInAppPurchases(const TArray<FString>& ProductIDs,
 	return AndroidThunkCpp_Iap_QueryInAppPurchases(ProductIDs);
 }
 
-bool AndroidThunkCpp_Iap_BeginPurchase(const FString& ProductID)
+bool AndroidThunkCpp_Iap_BeginPurchase(const FString& ProductID, const FString& AccountId)
 {
 	FPlatformMisc::LowLevelOutputDebugStringf(TEXT("[JNI] - AndroidThunkCpp_Iap_BeginPurchase %s"), *ProductID);
 	bool bResult = false;
@@ -1347,7 +1347,15 @@ bool AndroidThunkCpp_Iap_BeginPurchase(const FString& ProductID)
 		CHECK_JNI_METHOD(FJavaWrapper::AndroidThunkJava_IapBeginPurchase);
 
 		auto ProductIdJava = FJavaHelper::ToJavaString(Env, ProductID);
-		bResult = FJavaWrapper::CallBooleanMethod(Env, FJavaWrapper::GoogleServicesThis, FJavaWrapper::AndroidThunkJava_IapBeginPurchase, *ProductIdJava);
+		if (AccountId.IsEmpty())
+		{
+			bResult = FJavaWrapper::CallBooleanMethod(Env, FJavaWrapper::GoogleServicesThis, FJavaWrapper::AndroidThunkJava_IapBeginPurchase, *ProductIdJava, nullptr);
+		}
+		else
+		{
+			auto ObfuscatedAccountIdJava = FJavaHelper::ToJavaString(Env, AccountId);
+			bResult = FJavaWrapper::CallBooleanMethod(Env, FJavaWrapper::GoogleServicesThis, FJavaWrapper::AndroidThunkJava_IapBeginPurchase, *ProductIdJava, *ObfuscatedAccountIdJava);
+		}
 	}
 
 	return bResult;
@@ -1356,7 +1364,7 @@ bool AndroidThunkCpp_Iap_BeginPurchase(const FString& ProductID)
 bool AndroidThunkCpp_Iap_BeginPurchase(const FString& ProductID, const bool bConsumable)
 {
 	FPlatformMisc::LowLevelOutputDebugString(TEXT("AndroidThunkCpp_Iap_BeginPurchase DEPRECATED, won't use consumable flag"));
-	return AndroidThunkCpp_Iap_BeginPurchase(ProductID);
+	return AndroidThunkCpp_Iap_BeginPurchase(ProductID, FString());
 }
 
 bool AndroidThunkCpp_Iap_ConsumePurchase(const FString& ProductToken)
