- Add missing import

diff --git a/Engine/Build/Android/Java/src/com/epicgames/ue4/GameActivity.java.template b/Engine/Build/Android/Java/src/com/epicgames/ue4/GameActivity.java.template
--- a/Engine/Build/Android/Java/src/com/epicgames/ue4/GameActivity.java.template
+++ b/Engine/Build/Android/Java/src/com/epicgames/ue4/GameActivity.java.template
@@ -67,6 +67,7 @@ import java.io.InputStreamReader;
 import java.io.OutputStream;
 import java.io.OutputStreamWriter;
 import java.io.StringWriter;
+import java.security.NoSuchAlgorithmException;
 
 import android.annotation.TargetApi;
 
