Fix issue with launch screen on Android

#ue5
#android

diff --git a/Engine/Build/Android/Java/src/com/epicgames/unreal/GameActivity.java.template b/Engine/Build/Android/Java/src/com/epicgames/unreal/GameActivity.java.template
--- a/Engine/Build/Android/Java/src/com/epicgames/unreal/GameActivity.java.template
+++ b/Engine/Build/Android/Java/src/com/epicgames/unreal/GameActivity.java.template
@@ -185,7 +185,7 @@ import com.epicgames.unreal.ConsoleCmdReceiver;
 
 import android.os.Build;
 
-// TODO: use the resources from the UE4 lib project once we've got the packager up and running
+// TODO: use the resources from the Unreal lib project once we've got the packager up and running
 //import com.epicgames.unreal.R;
 
 import com.epicgames.unreal.DownloadShim;
@@ -2295,8 +2295,8 @@ public class GameActivity extends $${gameActivitySuperClass}$$ implements Surfac
 			{
 				SplashScreenLaunch = true;
 				try {
-					// try to get the splash theme (can't use R.style.UE4SplashTheme since we don't know the package name until runtime)
-					int SplashThemeId = getResources().getIdentifier("UE4SplashTheme", "style", getPackageName());
+					// try to get the splash theme (can't use R.style.UnrealSplashTheme since we don't know the package name until runtime)
+					int SplashThemeId = getResources().getIdentifier("UnrealSplashTheme", "style", getPackageName());
 					mSplashDialog = new Dialog(this, SplashThemeId);
 					mSplashDialog.setCancelable(false);
 
