Fully specify configrules key and provide CRC

#android

diff --git a/Engine/Build/Android/Java/src/com/epicgames/unreal/GameActivity.java.template b/Engine/Build/Android/Java/src/com/epicgames/unreal/GameActivity.java.template
--- a/Engine/Build/Android/Java/src/com/epicgames/unreal/GameActivity.java.template
+++ b/Engine/Build/Android/Java/src/com/epicgames/unreal/GameActivity.java.template
@@ -45,6 +45,8 @@ import java.util.Stack;
 import java.util.Timer;
 import java.util.TimerTask;
 import java.lang.Runnable;
+import java.util.zip.CRC32;
+import java.util.zip.Checksum;
 
 import java.io.BufferedReader;
 import java.io.BufferedWriter;
@@ -908,7 +910,7 @@ public class GameActivity extends $${gameActivitySuperClass}$$ implements Surfac
 			Field keyField = this.getClass().getDeclaredField("CONFIGRULES_KEY");
 			String key = (String)keyField.get(this);
 
-			Cipher cipher = Cipher.getInstance("AES");
+			Cipher cipher = Cipher.getInstance("AES/ECB/PKCS5Padding");
 			cipher.init(Cipher.DECRYPT_MODE, generateAESKey(key));
 			byte[] decrypted = cipher.doFinal(bytesToDecompress, headerSize, sizeCompressed);
 			sizeCompressed = decrypted.length;
@@ -970,6 +972,13 @@ public class GameActivity extends $${gameActivitySuperClass}$$ implements Surfac
 		byte[] configBytesInternal = getByteArrayFromFile(InternalFilesDir + "configrules.bin.png");
 		byte[] configBytesExternal = getByteArrayFromFile(android.os.Environment.getExternalStorageDirectory().getAbsolutePath() + "/UnrealGame/" + ProjectName + "/configrules.bin.png");
 
+		if (configBytesInternal != null)
+		{
+			CRC32 crc = new CRC32();
+			crc.update(configBytesInternal);
+			Log.debug(String.format("CRC32: %08X", crc.getValue()));
+		}
+
 		int assetsVersion = getConfigRulesVersion(configBytesAssets);
 		int internalVersion = getConfigRulesVersion(configBytesInternal);
 		int externalVersion = getConfigRulesVersion(configBytesExternal);
diff --git a/Engine/Build/Android/Prebuilt/ConfigRulesTool/compile.bat b/Engine/Build/Android/Prebuilt/ConfigRulesTool/compile.bat
--- a/Engine/Build/Android/Prebuilt/ConfigRulesTool/compile.bat
+++ b/Engine/Build/Android/Prebuilt/ConfigRulesTool/compile.bat
@@ -1,3 +1,3 @@
-%JAVA_HOME%\bin\javac.exe -d classes src\*.java
+"%JAVA_HOME%"\bin\javac.exe -d classes src\*.java
 cd classes
-%JAVA_HOME%\bin\jar cfe ../bin/ConfigRulesTool.jar ConfigRulesTool *.class
\ No newline at end of file
+"%JAVA_HOME%"\bin\jar cfe ../bin/ConfigRulesTool.jar ConfigRulesTool *.class
\ No newline at end of file
diff --git a/Engine/Build/Android/Prebuilt/ConfigRulesTool/src/ConfigRulesTool.java b/Engine/Build/Android/Prebuilt/ConfigRulesTool/src/ConfigRulesTool.java
--- a/Engine/Build/Android/Prebuilt/ConfigRulesTool/src/ConfigRulesTool.java
+++ b/Engine/Build/Android/Prebuilt/ConfigRulesTool/src/ConfigRulesTool.java
@@ -17,6 +17,8 @@ import javax.crypto.spec.SecretKeySpec;
 
 public class ConfigRulesTool
 {
+	public static String CipherTransform = "AES/ECB/PKCS5Padding";
+	
 	public static void writeInt(FileOutputStream outStream, int value)
 	{
 		try
@@ -135,7 +137,7 @@ public class ConfigRulesTool
 			{
 				try
 				{
-					Cipher cipher = Cipher.getInstance("AES");
+					Cipher cipher = Cipher.getInstance(CipherTransform);
 					cipher.init(Cipher.ENCRYPT_MODE, generateKey(key));
 					byte[] encrypted = cipher.doFinal(bytesCompressed, 0, sizeCompressed);
 					bytesCompressed = encrypted;
@@ -148,7 +150,7 @@ public class ConfigRulesTool
 					System.exit(-1);
 				}
 			}
-			
+
 			File outFile = new File(outFilename);
 			FileOutputStream fileOutStream = null;
 			try
@@ -182,6 +184,22 @@ public class ConfigRulesTool
 
 			System.out.println("Version: " + Integer.toString(version) + ", Compressed from " + Integer.toString(sizeUncompressed) +
 				" bytes to " + Integer.toString(sizeCompressed + headerSize) + " bytes" + (key.equals("") ? "." : " and encrypted."));
+
+			CRC32 crc = new CRC32();
+			path = Paths.get(outFilename);
+			byte[] bytesToCRC = null;
+			try
+			{
+				bytesToCRC = Files.readAllBytes(path);
+			}
+			catch (IOException e)
+			{
+				System.out.println("Unable to read file: " + outFilename);
+				System.exit(-1);
+			}
+			crc.update(bytesToCRC);
+			System.out.println(String.format("CRC32: %08X", crc.getValue()));
+
 			System.exit(0);
 		}
 	
@@ -221,7 +239,7 @@ public class ConfigRulesTool
 			{
 				try
 				{
-					Cipher cipher = Cipher.getInstance("AES");
+					Cipher cipher = Cipher.getInstance(CipherTransform);
 					cipher.init(Cipher.DECRYPT_MODE, generateKey(key));
 					byte[] decrypted = cipher.doFinal(bytesToDecompress, headerSize, sizeCompressed);
 					sizeCompressed = decrypted.length;
