Add check for ro.soc.model for configrules chipset matching

#android

diff --git a/Engine/Build/Android/Java/src/com/epicgames/unreal/GameActivity.java.template b/Engine/Build/Android/Java/src/com/epicgames/unreal/GameActivity.java.template
--- a/Engine/Build/Android/Java/src/com/epicgames/unreal/GameActivity.java.template
+++ b/Engine/Build/Android/Java/src/com/epicgames/unreal/GameActivity.java.template
@@ -1670,6 +1670,12 @@ public class GameActivity extends $${gameActivitySuperClass}$$ implements Surfac
 			hardware = cpuInfo.get("Hardware");
 		}
 		Log.debug("Hardware: " + hardware);
+		String roSOCModel = getProp("ro.soc.model");
+		if (!roSOCModel.equals(""))
+		{
+			Log.debug("ro.soc.model: " + roSOCModel);
+			hardware = roSOCModel;
+		}
 
 		// register variables for the program to access
 		Map<String, String> variables = new HashMap<String, String>();
@@ -1687,6 +1693,7 @@ public class GameActivity extends $${gameActivitySuperClass}$$ implements Surfac
 		variables.put("OverrideProfile", "true");
 		variables.put("hardware", hardware);
 		variables.put("ro.hardware", roHardware);
+		variables.put("ro.soc.model", roSOCModel);
 		variables.put("processor", processor);
 		variables.put("processorCount", Integer.toString(processorCount));
 		variables.put("useAffinity", "false");
diff --git a/Engine/Source/Developer/Android/AndroidDeviceDetection/Private/AndroidDeviceDetectionModule.cpp b/Engine/Source/Developer/Android/AndroidDeviceDetection/Private/AndroidDeviceDetectionModule.cpp
--- a/Engine/Source/Developer/Android/AndroidDeviceDetection/Private/AndroidDeviceDetectionModule.cpp
+++ b/Engine/Source/Developer/Android/AndroidDeviceDetection/Private/AndroidDeviceDetectionModule.cpp
@@ -509,6 +509,18 @@ private:
 					const TCHAR* Ptr = *RoHardware;
 					FParse::Line(&Ptr, NewDeviceInfo.Hardware);
 				}
+				{
+					HardwareCommand = FString::Printf(TEXT("-s %s %s ro.soc.model"), *NewDeviceInfo.SerialNumber, *GetPropCommand);
+					FString RoSOCModelIn;
+					FString RoSOCModelOut;
+					ExecuteAdbCommand(*HardwareCommand, &RoSOCModelIn, nullptr);
+					const TCHAR* Ptr = *RoSOCModelIn;
+					FParse::Line(&Ptr, RoSOCModelOut);
+					if (!RoSOCModelOut.IsEmpty())
+					{
+						NewDeviceInfo.Hardware = RoSOCModelOut;
+					}
+				}
 
 				// Read hardware from cpuinfo:
 				FString CPUInfoString;
