Add missing import

#android

diff --git a/Engine/Build/Android/Java/src/com/epicgames/unreal/GameActivity.java.template b/Engine/Build/Android/Java/src/com/epicgames/unreal/GameActivity.java.template
--- a/Engine/Build/Android/Java/src/com/epicgames/unreal/GameActivity.java.template
+++ b/Engine/Build/Android/Java/src/com/epicgames/unreal/GameActivity.java.template
@@ -32,6 +32,7 @@ import android.hardware.Sensor;
 import android.hardware.SensorEvent;
 import android.hardware.SensorEventListener;
 import java.lang.Override;
+import java.util.concurrent.atomic.AtomicBoolean;
 import java.util.Arrays;
 import java.util.List;
 import java.util.Map;
