Update ValidatePermissions utility

#android

diff --git a/Engine/Build/Android/Prebuilt/ValidatePermissions/compile.bat b/Engine/Build/Android/Prebuilt/ValidatePermissions/compile.bat
--- a/Engine/Build/Android/Prebuilt/ValidatePermissions/compile.bat
+++ b/Engine/Build/Android/Prebuilt/ValidatePermissions/compile.bat
@@ -1,3 +1,3 @@
-%JAVA_HOME%\bin\javac.exe -d classes src\*.java
+"%JAVA_HOME%\bin\javac.exe" -d classes src\*.java
 cd classes
-%JAVA_HOME%\bin\jar cfe ../bin/ValidatePermissions.jar ValidatePermissions *.class
\ No newline at end of file
+"%JAVA_HOME%\bin\jar" cfe ../bin/ValidatePermissions.jar ValidatePermissions *.class
\ No newline at end of file
diff --git a/Engine/Build/Android/Prebuilt/ValidatePermissions/src/ValidatePermissions.java b/Engine/Build/Android/Prebuilt/ValidatePermissions/src/ValidatePermissions.java
--- a/Engine/Build/Android/Prebuilt/ValidatePermissions/src/ValidatePermissions.java
+++ b/Engine/Build/Android/Prebuilt/ValidatePermissions/src/ValidatePermissions.java
@@ -10,7 +10,7 @@ import java.nio.file.Path;
 
 public class ValidatePermissions
 {
-	public static Map<String, String> ReadWhitelist(String filename)
+	public static Map<String, String> ReadAllowList(String filename)
 	{
 		Path path = Paths.get(filename);
 		if (Files.notExists(path))
@@ -38,11 +38,11 @@ public class ValidatePermissions
 
 	public static void main(String[] args)
 	{
-		String whitelistFilename = "permission_whitelist.txt";
-		Map<String, String> whitelist = ReadWhitelist(whitelistFilename);
-		if (whitelist == null)
+		String allowlistFilename = "permission_allowlist.txt";
+		Map<String, String> allowlist = ReadAllowList(allowlistFilename);
+		if (allowlist == null)
 		{
-			System.out.println("ValidatePermissions: permission_whitelist.txt not provided; not validating.");
+			System.out.println("ValidatePermissions: permission_allow.txt not provided; not validating.");
 			System.exit(0);
 		}
 
@@ -119,10 +119,10 @@ public class ValidatePermissions
 					}
 					line = line.substring(0, stopIndex);
 
-					// check against whitelist
-					if (!whitelist.containsKey(line))
+					// check against allowlist
+					if (!allowlist.containsKey(line))
 					{
-						System.out.println("ERROR: Permission '" + line + "' NOT whitelisted!");
+						System.out.println("ERROR: Permission '" + line + "' NOT allow listed!");
 						System.exit(-1);
 					}
 				}
