Add OGL PSO service's JNI functions to proguard.
Enables compile service optimizations in distribution builds.

diff --git a/Engine/Build/Android/Java/proguard-project.txt b/Engine/Build/Android/Java/proguard-project.txt
--- a/Engine/Build/Android/Java/proguard-project.txt
+++ b/Engine/Build/Android/Java/proguard-project.txt
@@ -207,6 +207,16 @@
 	public *;
 }
 
+-keep class com.epicgames.unreal.oglservices.OGLProgramServiceAccessor$JNIProgramLinkResponse 
+{
+	*;
+}
+
+-keepclassmembers, allowoptimization class com.epicgames.unreal.oglservices.OGLProgramServiceAccessor
+{
+    public static *** AndroidThunkJava_*(...);
+}
+
 -dontwarn com.google.vr.cardboard.AndroidNCompat
 -dontwarn com.google.vr.cardboard.StoragePermissionUtils
 -dontwarn sun.misc.Unsafe
\ No newline at end of file
