Update chipset detection to handle mismatches in ro.hardware and ro.soc.model
Add MT6785 chipset and fix a typo in MT6853

#android

diff --git a/Engine/Build/Android/Java/src/com/epicgames/unreal/GameActivity.java.template b/Engine/Build/Android/Java/src/com/epicgames/unreal/GameActivity.java.template
--- a/Engine/Build/Android/Java/src/com/epicgames/unreal/GameActivity.java.template
+++ b/Engine/Build/Android/Java/src/com/epicgames/unreal/GameActivity.java.template
@@ -1681,7 +1681,19 @@ public class GameActivity extends $${gameActivitySuperClass}$$ implements Surfac
 		if (!roSOCModel.equals(""))
 		{
 			Log.debug("ro.soc.model: " + roSOCModel);
-			hardware = roSOCModel;
+
+			// validate either we only have ro.soc.model, or that previous hardware value in agreement
+			// encountered a device which had ro.hardware and /proc/cpuinfo correct but ro.soc.model set differently
+			String CheckHardware = hardware.toLowerCase();
+			String CheckSOC = roSOCModel.toLowerCase();
+			if (CheckHardware.isEmpty() || CheckHardware.contains(CheckSOC) || CheckSOC.contains(CheckHardware))
+			{
+				hardware = roSOCModel;
+			}
+			else
+			{
+				Log.debug("Warning: hardware does not match ro.soc.model, assuming hardware is correct");
+			}
 		}
 
 		// register variables for the program to access
