Remove depreciated AppCache settings calls from WebViewControl

#android

diff --git a/Engine/Build/Android/Java/src/com/epicgames/unreal/WebViewControl.java b/Engine/Build/Android/Java/src/com/epicgames/unreal/WebViewControl.java
--- a/Engine/Build/Android/Java/src/com/epicgames/unreal/WebViewControl.java
+++ b/Engine/Build/Android/Java/src/com/epicgames/unreal/WebViewControl.java
@@ -119,10 +119,7 @@ class WebViewControl
 				webView.setWebChromeClient(new ChromeClient());
 				webView.getSettings().setJavaScriptEnabled(true);
 				webView.getSettings().setJavaScriptCanOpenWindowsAutomatically(true);
-				webView.getSettings().setAppCacheMaxSize( 10 * 1024 * 1024 );
-				webView.getSettings().setAppCachePath(GameActivity._activity.getApplicationContext().getCacheDir().getAbsolutePath() );
 				webView.getSettings().setAllowFileAccess( true );
-				webView.getSettings().setAppCacheEnabled( true );
 				webView.getSettings().setAllowContentAccess( true );
 				webView.getSettings().setAllowFileAccessFromFileURLs(true);
 				webView.getSettings().setAllowUniversalAccessFromFileURLs(true);
