Fix for process crash on PSOService shutdown

diff --git a/Engine/Build/Android/Prebuilt/PSOService/Source/PSOService.cpp b/Engine/Build/Android/Prebuilt/PSOService/Source/PSOService.cpp
--- a/Engine/Build/Android/Prebuilt/PSOService/Source/PSOService.cpp
+++ b/Engine/Build/Android/Prebuilt/PSOService/Source/PSOService.cpp
@@ -31,8 +31,8 @@ VKAPI_ATTR VkBool32 VKAPI_CALL VKValidationCallback(
 class FVulkanPSOCompiler
 {
 	bool bInitialized = false;
-	VkDevice Device;
-	VkInstance Instance;
+	VkDevice Device = VK_NULL_HANDLE;
+	VkInstance Instance = VK_NULL_HANDLE;
 
 	std::vector<VkPhysicalDevice> devices;
 	PFN_vkCreateRenderPass2KHR vkCreateRenderPass2;
@@ -203,6 +203,11 @@ public:
 
 	void ShutDownDevice()
 	{
+		if (Device == VK_NULL_HANDLE)
+		{
+			return;
+		}
+
 		if (PipelineCache != VK_NULL_HANDLE)
 		{
 			vkDestroyPipelineCache(Device, PipelineCache, nullptr);
@@ -763,7 +768,7 @@ JNI_METHOD void Java_com_epicgames_unreal_psoservices_PSOProgramService_InitVKDe
 
 }
 
-JNI_METHOD void Java_com_epicgames_unreal_psoservices_PSOProgramService_ShutDownVKDevice(JNIEnv* jenv, jobject thiz)
+JNI_METHOD void Java_com_epicgames_unreal_psoservices_PSOProgramService_ShutdownVKDevice(JNIEnv* jenv, jobject thiz)
 {
 	FVulkanPSOCompiler::Get().ShutDownDevice();
 }
