Expose versionCode and versionName to configrules

#android

diff --git a/Engine/Build/Android/Java/src/com/epicgames/unreal/GameActivity.java.template b/Engine/Build/Android/Java/src/com/epicgames/unreal/GameActivity.java.template
--- a/Engine/Build/Android/Java/src/com/epicgames/unreal/GameActivity.java.template
+++ b/Engine/Build/Android/Java/src/com/epicgames/unreal/GameActivity.java.template
@@ -238,6 +238,9 @@ public class GameActivity extends $${gameActivitySuperClass}$$ implements Surfac
 															ComponentCallbacks2
 {
 	public int VersionCode = 0;
+	public String VersionName = "Unknown";
+	public int targetSdkVersion = 0;
+
 	private boolean bAllowIMU = true;
 	private SensorManager sensorManager;
 	private Sensor accelerometer;
@@ -1708,6 +1711,9 @@ public class GameActivity extends $${gameActivitySuperClass}$$ implements Surfac
 		variables.put("memory", Integer.toString(totalMemory));
 
 		//variables.put("Profile", "Android");
+		variables.put("versionCode", Integer.toString(VersionCode));
+		variables.put("versionName", VersionName);
+		variables.put("targetSdkVersion", Integer.toString(targetSdkVersion));
 		variables.put("unsupportedDevice", "false");
 		variables.put("OverrideProfile", "true");
 		variables.put("hardware", hardware);
@@ -2394,18 +2400,20 @@ public class GameActivity extends $${gameActivitySuperClass}$$ implements Surfac
 			}
 		}
 
-		//Check for target sdk.  If 23 or higher then warn that permission handling may mean features don't work if user denies them.
-		int targetSdkVersion = 0;
+		// Retrieve version code, name, and target SDK
 		try 
 		{
 			PackageInfo packageInfo = getPackageManager().getPackageInfo(getPackageName(), 0);
+			VersionCode = packageInfo.versionCode;
+			VersionName = packageInfo.versionName;
 			targetSdkVersion = packageInfo.applicationInfo.targetSdkVersion;
 		}
-		catch (PackageManager.NameNotFoundException e) 
+		catch (Exception e)
 		{
-			Log.debug(e.getMessage());
+			Log.debug("Error accessing packageInfo: " + e.getMessage());
 		}
 
+		// Check for target sdk.  If 23 or higher then warn that permission handling may mean features don't work if user denies them.
 		boolean bIsShipping = nativeIsShippingBuild();
 		if (ANDROID_BUILD_VERSION >= 23 && targetSdkVersion >= 23) //23 is the API level (Marshmallow) where runtime permission handling is available
 		{
@@ -2645,7 +2653,7 @@ public class GameActivity extends $${gameActivitySuperClass}$$ implements Surfac
 				_activity.setRequestedOrientation( android.content.pm.ActivityInfo.SCREEN_ORIENTATION_LANDSCAPE );
 			}
 		}
-		
+
 		// Read metadata from AndroidManifest.xml
 		appPackageName = getPackageName();
 		String ProjectName = getPackageName();
@@ -2860,17 +2868,15 @@ public class GameActivity extends $${gameActivitySuperClass}$$ implements Surfac
 		Log.debug( "Debugger attached is " + bDebuggerAttached );
 
 		nativeSetAndroidVersionInformation( android.os.Build.VERSION.RELEASE, targetSdkVersion, android.os.Build.MANUFACTURER, android.os.Build.MODEL, android.os.Build.DISPLAY, Language );
-
-		try
+		int PatchVersion = 0;
+		if (VersionCode > 0)
 		{
-			VersionCode = getPackageManager().getPackageInfo(getPackageName(), 0).versionCode;
-			int PatchVersion = 0;
 			nativeSetObbInfo(ProjectName, getApplicationContext().getPackageName(), VersionCode, PatchVersion, AppType);
 		}
-		catch (Exception e)
+		else
 		{
 			// if the above failed, then, we can't use obbs
-			Log.debug("==================================> PackageInfo failure getting .obb info: " + e.getMessage());
+			Log.debug("==================================> PackageInfo failure getting .obb info: versionCode invalid");
 		}
 		
 		// enable the physical volume controls to the game
