Fix proguard exceptions for Android PSO compiler services JNI.

diff --git a/Engine/Build/Android/Java/proguard-project.txt b/Engine/Build/Android/Java/proguard-project.txt
--- a/Engine/Build/Android/Java/proguard-project.txt
+++ b/Engine/Build/Android/Java/proguard-project.txt
@@ -212,7 +212,7 @@
 	*;
 }
 
--keepclassmembers, allowoptimization class com.epicgames.unreal.psoservices.PSOServiceAccessor
+-keepclassmembers, allowoptimization class com.epicgames.unreal.psoservices.PSOProgramServiceAccessor
 {
     public static *** AndroidThunkJava_*(...);
 }
