Fixed: occasional crash at android.view.Surface.checkNotReleasedLocked(Surface.java:753) when app is backgrounded

diff --git a/Engine/Build/Android/Java/src/com/epicgames/unreal/GameActivity.java.template b/Engine/Build/Android/Java/src/com/epicgames/unreal/GameActivity.java.template
--- a/Engine/Build/Android/Java/src/com/epicgames/unreal/GameActivity.java.template
+++ b/Engine/Build/Android/Java/src/com/epicgames/unreal/GameActivity.java.template
@@ -5399,11 +5399,19 @@ public class GameActivity extends $${gameActivitySuperClass}$$ implements Surfac
 				@Override
 				public void run()
 				{
-					MySurfaceView.getHolder().getSurface().setFrameRate(FinalRefreshRate, android.view.Surface.FRAME_RATE_COMPATIBILITY_FIXED_SOURCE);
+					// the window might be gone by the time this code runs
+					try
+					{
+						MySurfaceView.getHolder().getSurface().setFrameRate(FinalRefreshRate, android.view.Surface.FRAME_RATE_COMPATIBILITY_FIXED_SOURCE);
+					}
+					catch (Exception e)
+					{
+						Log.debug("setFrameRate failed. "+ e.getMessage());
+					}
 				}
 			});
+			
 			Log.debug("SetNativeDisplayRefreshRate - setFrameRate " + RefreshRate);
-
 			return true;
 		}
 		else if(ANDROID_BUILD_VERSION >= 24)
diff --git a/Engine/Source/Runtime/VulkanRHI/Private/Android/VulkanAndroidPlatform.cpp b/Engine/Source/Runtime/VulkanRHI/Private/Android/VulkanAndroidPlatform.cpp
--- a/Engine/Source/Runtime/VulkanRHI/Private/Android/VulkanAndroidPlatform.cpp
+++ b/Engine/Source/Runtime/VulkanRHI/Private/Android/VulkanAndroidPlatform.cpp
@@ -492,7 +492,8 @@ bool FVulkanAndroidPlatform::FramePace(FVulkanDevice& Device, VkSwapchainKHR Swa
 {
 	bool bVsyncMultiple = true;
 	int32 CurrentFramePace = FAndroidPlatformRHIFramePacer::GetFramePace();
-	if (CurrentFramePace != 0)
+	// give a new swapchain a chance to present some frames before attempting to change refresh rate
+	if (CurrentFramePace != 0 && PresentID > FVulkanViewport::NUM_BUFFERS)
 	{
 		CachedRefreshRate = FAndroidMisc::GetNativeDisplayRefreshRate();
 		
