Fix missing flag on PendingIntent for Target SDK 31+

#android

robomerge[FN_EngineMerge] -UE5-RES

diff --git a/Engine/Build/Android/Java/JavaTemplates/DownloaderActivity.java.template b/Engine/Build/Android/Java/JavaTemplates/DownloaderActivity.java.template
--- a/Engine/Build/Android/Java/JavaTemplates/DownloaderActivity.java.template
+++ b/Engine/Build/Android/Java/JavaTemplates/DownloaderActivity.java.template
@@ -601,7 +601,7 @@ public class DownloaderActivity extends Activity implements IDownloaderClient {
 					PendingIntent pendingIntent = PendingIntent.getActivity(
 							DownloaderActivity.this,
 							0, intentToLaunchThisActivityFromNotification,
-							PendingIntent.FLAG_UPDATE_CURRENT);
+							PendingIntent.FLAG_UPDATE_CURRENT | PendingIntent.FLAG_IMMUTABLE);
 					// Request to start the download
 					int startResult = DownloaderClientMarshaller.startDownloadServiceIfRequired(this,
 							pendingIntent, OBBDownloaderService.class);
diff --git a/Engine/Build/Android/Java/src/com/epicgames/unreal/GameActivity.java.template b/Engine/Build/Android/Java/src/com/epicgames/unreal/GameActivity.java.template
--- a/Engine/Build/Android/Java/src/com/epicgames/unreal/GameActivity.java.template
+++ b/Engine/Build/Android/Java/src/com/epicgames/unreal/GameActivity.java.template
@@ -5088,7 +5088,7 @@ public class GameActivity extends $${gameActivitySuperClass}$$ implements Surfac
 		notificationIntent.putExtra("local-notification-activationEvent", activationEvent);
 
 		// Designate the callback as a PendingIntent
-		PendingIntent pendingIntent = PendingIntent.getBroadcast(context, notificationID, notificationIntent, PendingIntent.FLAG_UPDATE_CURRENT);
+		PendingIntent pendingIntent = PendingIntent.getBroadcast(context, notificationID, notificationIntent, PendingIntent.FLAG_UPDATE_CURRENT | PendingIntent.FLAG_IMMUTABLE);
 
 		TimeZone targetTimeZone = TimeZone.getTimeZone("UTC");
 
@@ -5158,7 +5158,7 @@ public class GameActivity extends $${gameActivitySuperClass}$$ implements Surfac
 		for(int curID : idList)
 		{
 			AlarmManager alarmManager = (AlarmManager)getSystemService(Context.ALARM_SERVICE);
-			PendingIntent pendingIntent = PendingIntent.getBroadcast(this, curID, new Intent(this, LocalNotificationReceiver.class), PendingIntent.FLAG_UPDATE_CURRENT);
+			PendingIntent pendingIntent = PendingIntent.getBroadcast(this, curID, new Intent(this, LocalNotificationReceiver.class), PendingIntent.FLAG_UPDATE_CURRENT | PendingIntent.FLAG_IMMUTABLE);
 			pendingIntent.cancel();
 			alarmManager.cancel(pendingIntent);
 		}
@@ -5186,7 +5186,7 @@ public class GameActivity extends $${gameActivitySuperClass}$$ implements Surfac
 
 			//Cancel the intent itself as well as from the alarm manager
 			AlarmManager alarmManager = (AlarmManager)getSystemService(Context.ALARM_SERVICE);
-			PendingIntent pendingIntent = PendingIntent.getBroadcast(this, notificationId, new Intent(this, LocalNotificationReceiver.class), PendingIntent.FLAG_UPDATE_CURRENT);
+			PendingIntent pendingIntent = PendingIntent.getBroadcast(this, notificationId, new Intent(this, LocalNotificationReceiver.class), PendingIntent.FLAG_UPDATE_CURRENT | PendingIntent.FLAG_IMMUTABLE);
 			pendingIntent.cancel();
 			alarmManager.cancel(pendingIntent);
 			
diff --git a/Engine/Build/Android/Java/src/com/epicgames/unreal/LocalNotificationReceiver.java b/Engine/Build/Android/Java/src/com/epicgames/unreal/LocalNotificationReceiver.java
--- a/Engine/Build/Android/Java/src/com/epicgames/unreal/LocalNotificationReceiver.java
+++ b/Engine/Build/Android/Java/src/com/epicgames/unreal/LocalNotificationReceiver.java
@@ -46,7 +46,7 @@ public class LocalNotificationReceiver extends BroadcastReceiver
 		notificationIntent.putExtra("localNotificationLaunchActivationEvent", activationEvent);
 
 		int notificationIconID = getNotificationIconID(context);
-		PendingIntent pendingNotificationIntent = PendingIntent.getActivity(context, notificationID, notificationIntent, 0);
+		PendingIntent pendingNotificationIntent = PendingIntent.getActivity(context, notificationID, notificationIntent, PendingIntent.FLAG_IMMUTABLE);
 
 		NotificationManager notificationManager = (NotificationManager) context.getSystemService(Context.NOTIFICATION_SERVICE);
 		@SuppressWarnings("deprecation")
diff --git a/Engine/Build/Android/Java/src/com/google/android/vending/expansion/downloader/impl/DownloaderService.java b/Engine/Build/Android/Java/src/com/google/android/vending/expansion/downloader/impl/DownloaderService.java
--- a/Engine/Build/Android/Java/src/com/google/android/vending/expansion/downloader/impl/DownloaderService.java
+++ b/Engine/Build/Android/Java/src/com/google/android/vending/expansion/downloader/impl/DownloaderService.java
@@ -956,7 +956,7 @@ public abstract class DownloaderService extends CustomIntentService implements I
         intent.setClassName(this.getPackageName(),
                 className);
         mAlarmIntent = PendingIntent.getBroadcast(this, 0, intent,
-                PendingIntent.FLAG_ONE_SHOT);
+                PendingIntent.FLAG_ONE_SHOT | PendingIntent.FLAG_IMMUTABLE);
         alarms.set(
                 AlarmManager.RTC_WAKEUP,
                 System.currentTimeMillis() + wakeUp, mAlarmIntent
diff --git a/Engine/Plugins/Online/Android/AndroidFetchBackgroundDownload/Source/AndroidFetchBackgroundDownload/Java/download/UEDownloadWorker.java b/Engine/Plugins/Online/Android/AndroidFetchBackgroundDownload/Source/AndroidFetchBackgroundDownload/Java/download/UEDownloadWorker.java
--- a/Engine/Plugins/Online/Android/AndroidFetchBackgroundDownload/Source/AndroidFetchBackgroundDownload/Java/download/UEDownloadWorker.java
+++ b/Engine/Plugins/Online/Android/AndroidFetchBackgroundDownload/Source/AndroidFetchBackgroundDownload/Java/download/UEDownloadWorker.java
@@ -244,7 +244,7 @@ public class UEDownloadWorker extends UEWorker implements DownloadProgressListen
 			notificationIntent.putExtra("localNotificationID" , Description.NotificationID);
 			notificationIntent.putExtra("localNotificationAppLaunched" , true);
 
-			pendingNotificationIntent = PendingIntent.getActivity(context, Description.NotificationID, notificationIntent, 0);
+			pendingNotificationIntent = PendingIntent.getActivity(context, Description.NotificationID, notificationIntent, PendingIntent.FLAG_IMMUTABLE);
 		}
 
 		Notification notification = new NotificationCompat.Builder(context, Description.NotificationChannelID)
diff --git a/Engine/Plugins/Runtime/Firebase/Source/Java/notifications/EpicFirebaseMessagingService.java b/Engine/Plugins/Runtime/Firebase/Source/Java/notifications/EpicFirebaseMessagingService.java
--- a/Engine/Plugins/Runtime/Firebase/Source/Java/notifications/EpicFirebaseMessagingService.java
+++ b/Engine/Plugins/Runtime/Firebase/Source/Java/notifications/EpicFirebaseMessagingService.java
@@ -212,7 +212,7 @@ public class EpicFirebaseMessagingService extends FirebaseMessagingService {
 
 	@NonNull
 	private PendingIntent getPendingIntentIntent(@NonNull Intent defaultIntent) {
-		return PendingIntent.getActivity(this, 1, defaultIntent, PendingIntent.FLAG_UPDATE_CURRENT);
+		return PendingIntent.getActivity(this, 1, defaultIntent, PendingIntent.FLAG_UPDATE_CURRENT | PendingIntent.FLAG_IMMUTABLE);
 	}
 
 	public void createChannel(@NonNull EpicFirebaseNotificationMeta meta) {
diff --git a/Engine/Plugins/Runtime/GoogleCloudMessaging/Source/Java/RemoteNotificationsListener.java b/Engine/Plugins/Runtime/GoogleCloudMessaging/Source/Java/RemoteNotificationsListener.java
--- a/Engine/Plugins/Runtime/GoogleCloudMessaging/Source/Java/RemoteNotificationsListener.java
+++ b/Engine/Plugins/Runtime/GoogleCloudMessaging/Source/Java/RemoteNotificationsListener.java
@@ -28,7 +28,7 @@ public class RemoteNotificationsListener extends GcmListenerService
 		{
 			notificationIconID = context.getResources().getIdentifier("icon", "drawable", context.getPackageName());
 		}
-		PendingIntent pendingNotificationIntent = PendingIntent.getActivity(this, notificationIconID, notificationIntent, 0);
+		PendingIntent pendingNotificationIntent = PendingIntent.getActivity(this, notificationIconID, notificationIntent, PendingIntent.FLAG_IMMUTABLE);
 
 		NotificationCompat.Builder builder = new NotificationCompat.Builder(this)
 			.setSmallIcon(notificationIconID)
