Use an alternative bitmap rendering mode for Quest webview

#android

diff --git a/Engine/Build/Android/Java/src/com/epicgames/unreal/WebViewControl.java b/Engine/Build/Android/Java/src/com/epicgames/unreal/WebViewControl.java
--- a/Engine/Build/Android/Java/src/com/epicgames/unreal/WebViewControl.java
+++ b/Engine/Build/Android/Java/src/com/epicgames/unreal/WebViewControl.java
@@ -17,6 +17,7 @@ import android.webkit.WebResourceResponse;
 import android.webkit.JsResult;
 import android.webkit.JsPromptResult;
 import android.graphics.Bitmap;
+import android.graphics.Bitmap.Config;
 import android.graphics.Paint;
 import android.graphics.Color;
 import android.webkit.WebBackForwardList;
@@ -76,12 +77,18 @@ class WebViewControl
 	private boolean VulkanRenderer = false;
 	private volatile boolean WaitOnBitmapRender = false;
 
+	private Bitmap TargetBitmap = null;
+	private Canvas TargetCanvas = null;
+	private java.nio.Buffer TargetBitmapFrameData = null;
+
+
 	private BitmapRenderer mBitmapRenderer = null;
 	private OESTextureRenderer mOESTextureRenderer = null;
 
 	public class FrameUpdateInfo 
 	{
 		public java.nio.Buffer Buffer;
+		public Bitmap Bitmap;
 		public boolean FrameReady;
 		public boolean RegionChanged;
 		public float UScale;
@@ -90,7 +97,7 @@ class WebViewControl
 		public float VOffset;
 	}
 
-	public WebViewControl(long inNativePtr, int width, int height, boolean swizzlePixels, boolean vulkanRenderer, final boolean bEnableRemoteDebugging, final boolean bUseTransparency, final boolean bEnableDomStorage)
+	public WebViewControl(long inNativePtr, int width, int height, boolean swizzlePixels, boolean vulkanRenderer, final boolean bEnableRemoteDebugging, final boolean bUseTransparency, final boolean bEnableDomStorage, final boolean bShouldUseBitmapRender)
 	{
 		final WebViewControl w = this;
 
@@ -102,6 +109,17 @@ class WebViewControl
 
 		nativePtr = inNativePtr;
 
+		if (bShouldUseBitmapRender)
+		{
+			TargetBitmap = Bitmap.createBitmap(initialWidth, initialHeight, Config.ARGB_8888);
+			TargetCanvas = new Canvas(TargetBitmap);
+		}
+		else
+		{
+			TargetBitmap = null;
+			TargetCanvas = null;
+		}
+
 		//GameActivity.Log.debug("WebViewControl width=" + width + ", height=" + height);
 
 		GameActivity._activity.runOnUiThread(new Runnable()
@@ -136,8 +154,9 @@ class WebViewControl
 				webView.getSettings().setLoadWithOverviewMode(true);
 				webView.getSettings().setUseWideViewPort(true);
 
-				//3D is the default
-				webView.SetAndroid3DBrowser(true);
+				//3D is the default if not using Bitmap
+				webView.SetBitmapDraw(TargetCanvas);
+				webView.SetAndroid3DBrowser(TargetBitmap == null);
 
 				if (bUseTransparency)
 				{
@@ -161,7 +180,7 @@ class WebViewControl
 	boolean PendingSetAndroid3DBrowser;
 	public void SetAndroid3DBrowser(boolean InIsAndroid3DBrowser)
 	{
-		PendingSetAndroid3DBrowser = InIsAndroid3DBrowser;
+		PendingSetAndroid3DBrowser = (TargetBitmap == null) ? InIsAndroid3DBrowser : false;
 		GameActivity._activity.runOnUiThread(new Runnable()
 		{
 			@Override
@@ -212,6 +231,16 @@ class WebViewControl
 			releaseOESTextureRenderer();
 		}
 
+		if (TargetCanvas != null)
+		{
+			TargetCanvas = null;
+		}
+		if (TargetBitmap != null)
+		{
+			TargetBitmap.recycle();
+			TargetBitmap = null;
+		}
+
 		Close();
 	}
 
@@ -349,7 +378,7 @@ class WebViewControl
 					// add to the activitiy, on top of the SurfaceView
 					ViewGroup.LayoutParams params = new ViewGroup.LayoutParams(ViewGroup.LayoutParams.FILL_PARENT, ViewGroup.LayoutParams.FILL_PARENT);			
 					GameActivity._activity.addContentView(positionLayout, params);
-					if(!webView.IsAndroid3DBrowser)
+					if(!webView.IsAndroid3DBrowser && TargetBitmap == null)
 					{
 						//GameActivity.Log.warn("request focus create");
 						webView.requestFocus();
@@ -408,7 +437,7 @@ class WebViewControl
 							NextURL = null;
 						}
 						else
-						if(!webView.IsAndroid3DBrowser)
+						if(!webView.IsAndroid3DBrowser && TargetBitmap == null)
 						{
 							//GameActivity.Log.warn("request focus");
 							webView.requestFocus();
@@ -504,8 +533,46 @@ class WebViewControl
 		}
 	}
 
+	public FrameUpdateInfo getVideoLastFrameBitmap()
+	{
+		if (TargetBitmap == null)
+		{
+			return null;
+		}
+
+		// trigger a draw (it will go to onDraw to do the actual update)
+		webView.draw(TargetCanvas);
+
+		FrameUpdateInfo frameUpdateInfo = new FrameUpdateInfo();
+
+		frameUpdateInfo.Bitmap = TargetBitmap;
+		frameUpdateInfo.FrameReady = true;
+		frameUpdateInfo.RegionChanged = false;
+		return frameUpdateInfo;
+	}
+
 	public FrameUpdateInfo getVideoLastFrameData()
 	{
+		if (TargetBitmap != null)
+		{
+			if (null == TargetBitmapFrameData)
+			{
+				TargetBitmapFrameData = java.nio.ByteBuffer.allocateDirect(initialWidth * initialHeight * 4);
+			}
+
+			// note RGBA vs BGRA so not really correct
+			TargetBitmapFrameData.position(0);
+			TargetBitmap.copyPixelsToBuffer(TargetBitmapFrameData);
+			TargetBitmapFrameData.position(0);
+
+			FrameUpdateInfo frameUpdateInfo = new FrameUpdateInfo();
+
+			frameUpdateInfo.Buffer = TargetBitmapFrameData;
+			frameUpdateInfo.FrameReady = true;
+			frameUpdateInfo.RegionChanged = false;
+			return frameUpdateInfo;
+		}
+
 		initBitmapRenderer();
 		if (null != mBitmapRenderer)
 		{
@@ -1557,6 +1624,8 @@ class WebViewControl
 	{
 		private android.view.Surface mSurface = null;
 		public boolean IsAndroid3DBrowser = false;
+		public boolean bBitmapDraw = false;
+		Canvas TargetCanvas = null;
 
 		// default constructors
 		public GLWebView(Context context) {
@@ -1580,12 +1649,26 @@ class WebViewControl
 			setOnTouchListener(new View.OnTouchListener() {
 				@Override
 				public boolean onTouch(View v, MotionEvent event) {
-					return IsAndroid3DBrowser;
+					return IsAndroid3DBrowser || bBitmapDraw;
 				}
 			});
 		
 		}
 
+		public void SetBitmapDraw(Canvas InCanvas)
+		{
+			boolean InBitmapDraw = InCanvas != null;
+			synchronized(this)
+			{
+				if(bBitmapDraw != InBitmapDraw)
+				{
+					TargetCanvas = InCanvas;
+					bBitmapDraw = InBitmapDraw;
+					webView.setFocusableInTouchMode(!(IsAndroid3DBrowser || bBitmapDraw));
+				}
+			}
+		}
+
 		public void SetAndroid3DBrowser(boolean InIsAndroid3DBrowser)
 		{
 			synchronized(this)
@@ -1593,7 +1676,7 @@ class WebViewControl
 				if(IsAndroid3DBrowser != InIsAndroid3DBrowser)
 				{
 					IsAndroid3DBrowser = InIsAndroid3DBrowser;
-					webView.setFocusableInTouchMode(!IsAndroid3DBrowser);
+					webView.setFocusableInTouchMode(!(IsAndroid3DBrowser || bBitmapDraw));
 				}
 			}
 		}
@@ -1608,6 +1691,17 @@ class WebViewControl
 		@Override
 		public void onDraw( Canvas canvas ) 
 		{
+			if (bBitmapDraw)
+			{
+				TargetCanvas.save();
+				float xScale = TargetCanvas.getWidth() / (float) canvas.getWidth();
+				TargetCanvas.translate(-getScrollX(), -getScrollY());
+				super.onDraw(TargetCanvas);
+				TargetCanvas.restore();
+				//GameActivity.Log.debug("onDraw: res: " + TargetCanvas.getWidth() + " x " + TargetCanvas.getHeight() + ", scroll:" + getScrollX() + " x " + getScrollY());
+				return;
+			}
+
 			//returns canvas attached to gl texture to draw on
 			Canvas glAttachedCanvas = null;
 			if(!IsAndroid3DBrowser)
diff --git a/Engine/Source/Runtime/WebBrowser/Private/Android/AndroidJavaWebBrowser.cpp b/Engine/Source/Runtime/WebBrowser/Private/Android/AndroidJavaWebBrowser.cpp
--- a/Engine/Source/Runtime/WebBrowser/Private/Android/AndroidJavaWebBrowser.cpp
+++ b/Engine/Source/Runtime/WebBrowser/Private/Android/AndroidJavaWebBrowser.cpp
@@ -6,6 +6,8 @@
 
 #include "Android/AndroidApplication.h"
 
+#include <android/bitmap.h>
+
 #if UE_BUILD_SHIPPING
 // always clear any exceptions in SHipping
 #define CHECK_JNI_RESULT(Id) if (Id == 0) { JEnv->ExceptionClear(); }
@@ -31,9 +33,10 @@ static jfieldID FindField(JNIEnv* JEnv, jclass Class, const ANSICHAR* FieldName,
 #define MOTIONEVENT_ACTION_MOVE		2
 
 FJavaAndroidWebBrowser::FJavaAndroidWebBrowser(bool swizzlePixels, bool vulkanRenderer, int32 width, int32 height,
-	jlong widgetPtr, bool bEnableRemoteDebugging, bool bUseTransparency, bool bEnableDomStorage)
-	: FJavaClassObject(GetClassName(), "(JIIZZZZZ)V", widgetPtr, width, height,  swizzlePixels, vulkanRenderer, bEnableRemoteDebugging, bUseTransparency, bEnableDomStorage)
+	jlong widgetPtr, bool bEnableRemoteDebugging, bool bUseTransparency, bool bEnableDomStorage, bool bShouldUseBitmapRender)
+	: FJavaClassObject(GetClassName(), "(JIIZZZZZZ)V", widgetPtr, width, height,  swizzlePixels, vulkanRenderer, bEnableRemoteDebugging, bUseTransparency, bEnableDomStorage, bShouldUseBitmapRender)
 	, ReleaseMethod(GetClassMethod("release", "()V"))
+	, GetVideoLastFrameBitmapMethod(GetClassMethod("getVideoLastFrameBitmap", "()Lcom/epicgames/unreal/WebViewControl$FrameUpdateInfo;"))
 	, GetVideoLastFrameDataMethod(GetClassMethod("getVideoLastFrameData", "()Lcom/epicgames/unreal/WebViewControl$FrameUpdateInfo;"))
 	, GetVideoLastFrameMethod(GetClassMethod("getVideoLastFrame", "(I)Lcom/epicgames/unreal/WebViewControl$FrameUpdateInfo;"))
 	, DidResolutionChangeMethod(GetClassMethod("didResolutionChange", "()Z"))
@@ -58,6 +61,7 @@ FJavaAndroidWebBrowser::FJavaAndroidWebBrowser(bool swizzlePixels, bool vulkanRe
 	// get field IDs for FrameUpdateInfo class members
 	FrameUpdateInfoClass = FAndroidApplication::FindJavaClassGlobalRef("com/epicgames/unreal/WebViewControl$FrameUpdateInfo");
 	FrameUpdateInfo_Buffer = FindField(JEnv, FrameUpdateInfoClass, "Buffer", "Ljava/nio/Buffer;", false);
+	FrameUpdateInfo_Bitmap = FindField(JEnv, FrameUpdateInfoClass, "Bitmap", "Landroid/graphics/Bitmap;", false);
 	FrameUpdateInfo_FrameReady = FindField(JEnv, FrameUpdateInfoClass, "FrameReady", "Z", false);
 	FrameUpdateInfo_RegionChanged = FindField(JEnv, FrameUpdateInfoClass, "RegionChanged", "Z", false);
 }
@@ -75,6 +79,66 @@ void FJavaAndroidWebBrowser::Release()
 	CallMethod<void>(ReleaseMethod);
 }
 
+bool FJavaAndroidWebBrowser::GetVideoLastFrameBitmap(void* outPixels, int64 outCount)
+{
+	// This can return an exception in some cases
+	JNIEnv* JEnv = FAndroidApplication::GetJavaEnv();
+	auto Result = NewScopedJavaObject(JEnv, JEnv->CallObjectMethod(Object, GetVideoLastFrameBitmapMethod.Method));
+	if (JEnv->ExceptionCheck())
+	{
+		JEnv->ExceptionDescribe();
+		JEnv->ExceptionClear();
+		return false;
+	}
+
+	if (!Result)
+	{
+		return false;
+	}
+
+	auto bitmap = NewScopedJavaObject(JEnv, JEnv->GetObjectField(*Result, FrameUpdateInfo_Bitmap));
+	if (bitmap)
+	{
+		AndroidBitmapInfo BitmapInfo;
+
+		AndroidBitmap_getInfo(JEnv, *bitmap, &BitmapInfo);
+		if (BitmapInfo.format != ANDROID_BITMAP_FORMAT_RGBA_8888)
+		{
+			return false;
+		}
+		int32 BitmapBytes = BitmapInfo.stride * BitmapInfo.height;
+
+		BitmapBytes = FMath::Min(outCount, BitmapBytes);
+		if (BitmapBytes > 0)
+		{
+			void* BitmapDataPtr;
+
+			if (AndroidBitmap_lockPixels(JEnv, *bitmap, &BitmapDataPtr) >= 0)
+			{
+				// unfortunately the order is wrong (RGBA instead of BGRA)
+				//memcpy(outPixels, BitmapDataPtr, BitmapBytes);
+				int32 NumPixels = BitmapBytes >> 2;
+				uint8* SrcData = (uint8*)BitmapDataPtr;
+				uint8* DstData = (uint8*)outPixels;
+				while (NumPixels--)
+				{
+					*DstData++ = SrcData[2];
+					*DstData++ = SrcData[1];
+					*DstData++ = SrcData[0];
+					*DstData++ = SrcData[3];
+					SrcData += 4;
+				}
+				AndroidBitmap_unlockPixels(JEnv, *bitmap);
+				return true;
+			}
+		}
+
+		return false;
+	}
+
+	return false;
+}
+
 bool FJavaAndroidWebBrowser::GetVideoLastFrameData(void* & outPixels, int64 & outCount, bool *bRegionChanged)
 {
 	// This can return an exception in some cases
diff --git a/Engine/Source/Runtime/WebBrowser/Private/Android/AndroidJavaWebBrowser.h b/Engine/Source/Runtime/WebBrowser/Private/Android/AndroidJavaWebBrowser.h
--- a/Engine/Source/Runtime/WebBrowser/Private/Android/AndroidJavaWebBrowser.h
+++ b/Engine/Source/Runtime/WebBrowser/Private/Android/AndroidJavaWebBrowser.h
@@ -15,9 +15,10 @@
 class FJavaAndroidWebBrowser : public FJavaClassObject
 {
 public:
-	FJavaAndroidWebBrowser(bool swizzlePixels, bool vulkanRenderer, int32 width, int32 height, jlong widgetPtr, bool bEnableRemoteDebugging, bool bUseTransparency, bool bEnableDomStorage);
+	FJavaAndroidWebBrowser(bool swizzlePixels, bool vulkanRenderer, int32 width, int32 height, jlong widgetPtr, bool bEnableRemoteDebugging, bool bUseTransparency, bool bEnableDomStorage, bool bShouldUseBitmapRender);
 	virtual ~FJavaAndroidWebBrowser();
 	void Release();
+	bool GetVideoLastFrameBitmap(void* outPixels, int64 outCount);
 	bool GetVideoLastFrameData(void* & outPixels, int64 & outCount, bool *bRegionChanged);
 	bool GetVideoLastFrame(int32 destTexture);
 	bool DidResolutionChange();
@@ -40,6 +41,7 @@ private:
 	static FName GetClassName();
 
 	FJavaClassMethod ReleaseMethod;
+	FJavaClassMethod GetVideoLastFrameBitmapMethod;
 	FJavaClassMethod GetVideoLastFrameDataMethod;
 	FJavaClassMethod GetVideoLastFrameMethod;
 	FJavaClassMethod DidResolutionChangeMethod;
@@ -59,6 +61,7 @@ private:
 	// FrameUpdateInfo member field ids
 	jclass FrameUpdateInfoClass;
 	jfieldID FrameUpdateInfo_Buffer;
+	jfieldID FrameUpdateInfo_Bitmap;
 	jfieldID FrameUpdateInfo_FrameReady;
 	jfieldID FrameUpdateInfo_RegionChanged;
 	
diff --git a/Engine/Source/Runtime/WebBrowser/Private/Android/AndroidWebBrowserWidget.cpp b/Engine/Source/Runtime/WebBrowser/Private/Android/AndroidWebBrowserWidget.cpp
--- a/Engine/Source/Runtime/WebBrowser/Private/Android/AndroidWebBrowserWidget.cpp
+++ b/Engine/Source/Runtime/WebBrowser/Private/Android/AndroidWebBrowserWidget.cpp
@@ -29,6 +29,8 @@
 
 #include <jni.h>
 
+extern bool AndroidThunkCpp_IsOculusMobileApplication();
+
 FCriticalSection SAndroidWebBrowserWidget::WebControlsCS;
 TMap<int64, TWeakPtr<SAndroidWebBrowserWidget>> SAndroidWebBrowserWidget::AllWebControls;
 
@@ -112,6 +114,9 @@ void SAndroidWebBrowserWidget::Construct(const FArguments& Args)
 	WebBrowserWindowPtr = Args._WebBrowserWindow;
 	IsAndroid3DBrowser = true;
 
+	bShouldUseBitmapRender = AndroidThunkCpp_IsOculusMobileApplication();
+	bMouseCapture = false;
+
 	HistorySize = 0;
 	HistoryPosition = 0;
 
@@ -123,7 +128,7 @@ void SAndroidWebBrowserWidget::Construct(const FArguments& Args)
 	FPlatformMisc::LowLevelOutputDebugStringf(TEXT("SAndroidWebBrowserWidget::Construct viewport=%d x %d"), viewportSize.X, viewportSize.Y);
 
 	JavaWebBrowser = MakeShared<FJavaAndroidWebBrowser, ESPMode::ThreadSafe>(false, FAndroidMisc::ShouldUseVulkan(), viewportSize.X, viewportSize.Y,
-		reinterpret_cast<jlong>(this), !(UE_BUILD_SHIPPING || UE_BUILD_TEST), Args._UseTransparency, bEnableDomStorage);
+		reinterpret_cast<jlong>(this), !(UE_BUILD_SHIPPING || UE_BUILD_TEST), Args._UseTransparency, bEnableDomStorage, bShouldUseBitmapRender);
 
 	TextureSamplePool = new FWebBrowserTextureSamplePool();
 	WebBrowserTextureSamplesQueue = MakeShared<FWebBrowserTextureSampleQueue, ESPMode::ThreadSafe>();
@@ -271,41 +276,69 @@ void SAndroidWebBrowserWidget::Tick(const FGeometry& AllottedGeometry, const dou
 			}
 			WriteWebBrowserParams = { JavaWebBrowser, WebBrowserTextureSamplesQueue, NewTextureSample, (int32)(viewportSize.X * viewportSize.Y * sizeof(int32)) };
 
-			ENQUEUE_RENDER_COMMAND(WriteAndroidWebBrowser)(
-				[Params = WriteWebBrowserParams](FRHICommandListImmediate& RHICmdList)
-				{
-					auto PinnedJavaWebBrowser = Params.JavaWebBrowserPtr.Pin();
-					auto PinnedSamples = Params.WebBrowserTextureSampleQueuePtr.Pin();
-
-					if (!PinnedJavaWebBrowser.IsValid() || !PinnedSamples.IsValid())
+			if (bShouldUseBitmapRender)
+			{
+				ENQUEUE_RENDER_COMMAND(WriteAndroidWebBrowser)(
+					[Params = WriteWebBrowserParams](FRHICommandListImmediate& RHICmdList)
 					{
-						return;
-					}
+						auto PinnedJavaWebBrowser = Params.JavaWebBrowserPtr.Pin();
+						auto PinnedSamples = Params.WebBrowserTextureSampleQueuePtr.Pin();
 
-					bool bRegionChanged = false;
+						if (!PinnedJavaWebBrowser.IsValid() || !PinnedSamples.IsValid())
+						{
+							return;
+						}
 
-					// write frame into buffer
-					void* Buffer = nullptr;
-					int64 SampleCount = 0;
+						int32 SampleBufferSize = Params.NewTextureSamplePtr->InitializeBufferForCopy();
+						void* Buffer = (void*)Params.NewTextureSamplePtr->GetBuffer();
 
-					if (!PinnedJavaWebBrowser->GetVideoLastFrameData(Buffer, SampleCount, &bRegionChanged))
-					{
-						//FPlatformMisc::LowLevelOutputDebugStringf(TEXT("Fetch RT: ShouldUseVulkan couldn't get texture buffer"));
-						return;
-					}
+						if (!PinnedJavaWebBrowser->GetVideoLastFrameBitmap(Buffer, SampleBufferSize))
+						{
+							return;
+						}
 
-					if (SampleCount != Params.SampleCount)
+						PinnedSamples->RequestFlush();
+						PinnedSamples->Enqueue(Params.NewTextureSamplePtr);
+					});
+			}
+			else
+			{
+				ENQUEUE_RENDER_COMMAND(WriteAndroidWebBrowser)(
+					[Params = WriteWebBrowserParams](FRHICommandListImmediate& RHICmdList)
 					{
-						FPlatformMisc::LowLevelOutputDebugStringf(TEXT("SAndroidWebBrowserWidget::Fetch: Sample count mismatch (Buffer=%llu, Available=%llu"), Params.SampleCount, SampleCount);
-					}
-					check(Params.SampleCount <= SampleCount);
+						auto PinnedJavaWebBrowser = Params.JavaWebBrowserPtr.Pin();
+						auto PinnedSamples = Params.WebBrowserTextureSampleQueuePtr.Pin();
 
-					// must make a copy (buffer is owned by Java, not us!)
-					Params.NewTextureSamplePtr->InitializeBuffer(Buffer, true);
+						if (!PinnedJavaWebBrowser.IsValid() || !PinnedSamples.IsValid())
+						{
+							return;
+						}
 
-					PinnedSamples->RequestFlush();
-					PinnedSamples->Enqueue(Params.NewTextureSamplePtr);
-				});
+						bool bRegionChanged = false;
+
+						// write frame into buffer
+						void* Buffer = nullptr;
+						int64 SampleCount = 0;
+
+						if (!PinnedJavaWebBrowser->GetVideoLastFrameData(Buffer, SampleCount, &bRegionChanged))
+						{
+							//FPlatformMisc::LowLevelOutputDebugStringf(TEXT("Fetch RT: ShouldUseVulkan couldn't get texture buffer"));
+							return;
+						}
+
+						if (SampleCount != Params.SampleCount)
+						{
+							FPlatformMisc::LowLevelOutputDebugStringf(TEXT("SAndroidWebBrowserWidget::Fetch: Sample count mismatch (Buffer=%llu, Available=%llu"), Params.SampleCount, SampleCount);
+						}
+						check(Params.SampleCount <= SampleCount);
+
+						// must make a copy (buffer is owned by Java, not us!)
+						Params.NewTextureSamplePtr->InitializeBuffer(Buffer, true);
+
+						PinnedSamples->RequestFlush();
+						PinnedSamples->Enqueue(Params.NewTextureSamplePtr);
+					});
+			}
 		}
 		else if (GSupportsImageExternal && WebBrowserTexture != nullptr)
 		{
@@ -530,6 +563,7 @@ FReply SAndroidWebBrowserWidget::OnMouseButtonDown(const FGeometry& MyGeometry,
 	{
 		Reply = FReply::Handled();
 		SendTouchDown(ConvertMouseEventToLocal(MyGeometry, MouseEvent));
+		bMouseCapture = true;
 	}
 
 	return Reply;
@@ -546,6 +580,7 @@ FReply SAndroidWebBrowserWidget::OnMouseButtonUp(const FGeometry& MyGeometry, co
 	{
 		Reply = FReply::Handled();
 		SendTouchUp(ConvertMouseEventToLocal(MyGeometry, MouseEvent));
+		bMouseCapture = false;
 	}
 
 	return Reply;
@@ -555,10 +590,7 @@ FReply SAndroidWebBrowserWidget::OnMouseMove(const FGeometry& MyGeometry, const
 {
 	FReply Reply = FReply::Unhandled();
 
-	FKey Button = MouseEvent.GetEffectingButton();
-	bool bSupportedButton = (Button == EKeys::LeftMouseButton); // || Button == EKeys::RightMouseButton || Button == EKeys::MiddleMouseButton);
-
-	if (bSupportedButton)
+	if (bMouseCapture)
 	{
 		Reply = FReply::Handled();
 		SendTouchMove(ConvertMouseEventToLocal(MyGeometry, MouseEvent));
diff --git a/Engine/Source/Runtime/WebBrowser/Private/Android/AndroidWebBrowserWidget.h b/Engine/Source/Runtime/WebBrowser/Private/Android/AndroidWebBrowserWidget.h
--- a/Engine/Source/Runtime/WebBrowser/Private/Android/AndroidWebBrowserWidget.h
+++ b/Engine/Source/Runtime/WebBrowser/Private/Android/AndroidWebBrowserWidget.h
@@ -109,6 +109,12 @@ private:
 	/** Enable 3D appearance for Android. */
 	bool IsAndroid3DBrowser;
 
+	/** Should use bitmap rendering for 3D. */
+	bool bShouldUseBitmapRender;
+
+	/** Mouse captured */
+	bool bMouseCapture;
+
 	/** The Java side webbrowser interface. */
 	TSharedPtr<FJavaAndroidWebBrowser, ESPMode::ThreadSafe> JavaWebBrowser;
 
diff --git a/Engine/Source/Runtime/WebBrowser/WebBrowser.Build.cs b/Engine/Source/Runtime/WebBrowser/WebBrowser.Build.cs
--- a/Engine/Source/Runtime/WebBrowser/WebBrowser.Build.cs
+++ b/Engine/Source/Runtime/WebBrowser/WebBrowser.Build.cs
@@ -60,6 +60,11 @@ public class WebBrowser : ModuleRules
 
 		}
 
+		if (Target.Platform == UnrealTargetPlatform.Android)
+		{
+			PublicSystemLibraries.Add("libjnigraphics");
+		}
+
 		if (Target.Platform == UnrealTargetPlatform.Win64
 		||  Target.Platform == UnrealTargetPlatform.Mac
 		||  Target.Platform == UnrealTargetPlatform.Linux)
diff --git a/Engine/Source/Runtime/WebBrowserTexture/Public/WebBrowserTextureSample.h b/Engine/Source/Runtime/WebBrowserTexture/Public/WebBrowserTextureSample.h
--- a/Engine/Source/Runtime/WebBrowserTexture/Public/WebBrowserTextureSample.h
+++ b/Engine/Source/Runtime/WebBrowserTexture/Public/WebBrowserTextureSample.h
@@ -109,6 +109,31 @@ public:
 		}
 	}
 
+	/**
+	* Initialize the sample for copy externally.
+	*
+	* @return The size of the buffer
+	* @see InitializeTexture
+	*/
+	int32 InitializeBufferForCopy()
+	{
+		SIZE_T RequiredBufferSize = Dim.X * Dim.Y * sizeof(int32);
+		if (BufferSize < RequiredBufferSize)
+		{
+			if (BufferSize == 0)
+			{
+				Buffer = FMemory::Malloc(RequiredBufferSize);
+			}
+			else
+			{
+				Buffer = FMemory::Realloc(Buffer, RequiredBufferSize);
+			}
+
+			BufferSize = RequiredBufferSize;
+		}
+		return BufferSize;
+	}
+
 	/**
 	* Initialize the sample with a texture resource.
 	*
