Add key support for Android WebBrowser widget

#android

diff --git a/Engine/Build/Android/Java/src/com/epicgames/unreal/WebViewControl.java b/Engine/Build/Android/Java/src/com/epicgames/unreal/WebViewControl.java
--- a/Engine/Build/Android/Java/src/com/epicgames/unreal/WebViewControl.java
+++ b/Engine/Build/Android/Java/src/com/epicgames/unreal/WebViewControl.java
@@ -7,6 +7,7 @@ import android.content.Context;
 
 import android.view.View;
 import android.view.ViewGroup;
+import android.view.KeyEvent;
 import android.webkit.ConsoleMessage;
 import android.webkit.WebView;
 import android.webkit.WebViewClient;
@@ -358,6 +359,173 @@ class WebViewControl
 		});
 	}
 
+	private static class KeyTableData
+	{
+		public final int key;
+		public final int meta;
+
+		KeyTableData (int keycode, int metadata)
+		{
+			key = keycode;
+			meta = metadata;
+		}
+	}
+	
+	private static final KeyTableData[] KeyTableList = {
+		new KeyTableData(KeyEvent.KEYCODE_UNKNOWN, 0),	// 0
+		new KeyTableData(KeyEvent.KEYCODE_UNKNOWN, 0),	// 1
+		new KeyTableData(KeyEvent.KEYCODE_UNKNOWN, 0),	// 2
+		new KeyTableData(KeyEvent.KEYCODE_UNKNOWN, 0),	// 3
+		new KeyTableData(KeyEvent.KEYCODE_UNKNOWN, 0),	// 4
+		new KeyTableData(KeyEvent.KEYCODE_UNKNOWN, 0),	// 5
+		new KeyTableData(KeyEvent.KEYCODE_UNKNOWN, 0),	// 6
+		new KeyTableData(KeyEvent.KEYCODE_UNKNOWN, 0),	// 7
+		new KeyTableData(KeyEvent.KEYCODE_DEL, 0),	// 8
+		new KeyTableData(KeyEvent.KEYCODE_TAB, 0),	// 9
+		new KeyTableData(KeyEvent.KEYCODE_ENTER, 0),	// 10
+		new KeyTableData(KeyEvent.KEYCODE_UNKNOWN, 0),	// 11
+		new KeyTableData(KeyEvent.KEYCODE_UNKNOWN, 0),	// 12
+		new KeyTableData(KeyEvent.KEYCODE_ENTER, 0),	// 13
+		new KeyTableData(KeyEvent.KEYCODE_UNKNOWN, 0),	// 14
+		new KeyTableData(KeyEvent.KEYCODE_UNKNOWN, 0),	// 15
+		new KeyTableData(KeyEvent.KEYCODE_UNKNOWN, 0),	// 16
+		new KeyTableData(KeyEvent.KEYCODE_UNKNOWN, 0),	// 17
+		new KeyTableData(KeyEvent.KEYCODE_UNKNOWN, 0),	// 18
+		new KeyTableData(KeyEvent.KEYCODE_UNKNOWN, 0),	// 19
+		new KeyTableData(KeyEvent.KEYCODE_UNKNOWN, 0),	// 20
+		new KeyTableData(KeyEvent.KEYCODE_UNKNOWN, 0),	// 21
+		new KeyTableData(KeyEvent.KEYCODE_UNKNOWN, 0),	// 22
+		new KeyTableData(KeyEvent.KEYCODE_UNKNOWN, 0),	// 23
+		new KeyTableData(KeyEvent.KEYCODE_UNKNOWN, 0),	// 24
+		new KeyTableData(KeyEvent.KEYCODE_UNKNOWN, 0),	// 25
+		new KeyTableData(KeyEvent.KEYCODE_UNKNOWN, 0),	// 26
+		new KeyTableData(KeyEvent.KEYCODE_UNKNOWN, 0),	// 27
+		new KeyTableData(KeyEvent.KEYCODE_UNKNOWN, 0),	// 28
+		new KeyTableData(KeyEvent.KEYCODE_UNKNOWN, 0),	// 29
+		new KeyTableData(KeyEvent.KEYCODE_UNKNOWN, 0),	// 30
+		new KeyTableData(KeyEvent.KEYCODE_UNKNOWN, 0),	// 31
+		new KeyTableData(KeyEvent.KEYCODE_SPACE, 0),	// 32
+		new KeyTableData(KeyEvent.KEYCODE_1, KeyEvent.META_SHIFT_ON),	// 33
+		new KeyTableData(KeyEvent.KEYCODE_APOSTROPHE, KeyEvent.META_SHIFT_ON),	// 34
+		new KeyTableData(KeyEvent.KEYCODE_POUND, 0),	// 35
+		new KeyTableData(KeyEvent.KEYCODE_4, KeyEvent.META_SHIFT_ON),	// 36
+		new KeyTableData(KeyEvent.KEYCODE_5, KeyEvent.META_SHIFT_ON),	// 37
+		new KeyTableData(KeyEvent.KEYCODE_7, KeyEvent.META_SHIFT_ON),	// 38
+		new KeyTableData(KeyEvent.KEYCODE_APOSTROPHE, 0),	// 39
+		new KeyTableData(KeyEvent.KEYCODE_NUMPAD_LEFT_PAREN, 0),	// 40
+		new KeyTableData(KeyEvent.KEYCODE_NUMPAD_RIGHT_PAREN, 0),	// 41
+		new KeyTableData(KeyEvent.KEYCODE_STAR, 0),	// 42
+		new KeyTableData(KeyEvent.KEYCODE_PLUS, 0),	// 43
+		new KeyTableData(KeyEvent.KEYCODE_COMMA, 0),	// 44
+		new KeyTableData(KeyEvent.KEYCODE_MINUS, 0),	// 45
+		new KeyTableData(KeyEvent.KEYCODE_PERIOD, 0),	// 46
+		new KeyTableData(KeyEvent.KEYCODE_SLASH, 0),	// 47
+		new KeyTableData(KeyEvent.KEYCODE_0, 0),	// 48
+		new KeyTableData(KeyEvent.KEYCODE_1, 0),	// 49
+		new KeyTableData(KeyEvent.KEYCODE_2, 0),	// 50
+		new KeyTableData(KeyEvent.KEYCODE_3, 0),	// 51
+		new KeyTableData(KeyEvent.KEYCODE_4, 0),	// 52
+		new KeyTableData(KeyEvent.KEYCODE_5, 0),	// 53
+		new KeyTableData(KeyEvent.KEYCODE_6, 0),	// 54
+		new KeyTableData(KeyEvent.KEYCODE_7, 0),	// 55
+		new KeyTableData(KeyEvent.KEYCODE_8, 0),	// 56
+		new KeyTableData(KeyEvent.KEYCODE_9, 0),	// 57
+		new KeyTableData(KeyEvent.KEYCODE_SEMICOLON, KeyEvent.META_SHIFT_ON),	// 58
+		new KeyTableData(KeyEvent.KEYCODE_SEMICOLON, 0),	// 59
+		new KeyTableData(KeyEvent.KEYCODE_COMMA, KeyEvent.META_SHIFT_ON),	// 60
+		new KeyTableData(KeyEvent.KEYCODE_EQUALS, 0),	// 61
+		new KeyTableData(KeyEvent.KEYCODE_PERIOD, KeyEvent.META_SHIFT_ON),	// 62
+		new KeyTableData(KeyEvent.KEYCODE_SLASH, KeyEvent.META_SHIFT_ON),	// 63
+		new KeyTableData(KeyEvent.KEYCODE_AT, 0),	// 64
+		new KeyTableData(KeyEvent.KEYCODE_A, KeyEvent.META_SHIFT_ON),	// 65
+		new KeyTableData(KeyEvent.KEYCODE_B, KeyEvent.META_SHIFT_ON),	// 66
+		new KeyTableData(KeyEvent.KEYCODE_C, KeyEvent.META_SHIFT_ON),	// 67
+		new KeyTableData(KeyEvent.KEYCODE_D, KeyEvent.META_SHIFT_ON),	// 68
+		new KeyTableData(KeyEvent.KEYCODE_E, KeyEvent.META_SHIFT_ON),	// 69
+		new KeyTableData(KeyEvent.KEYCODE_F, KeyEvent.META_SHIFT_ON),	// 70
+		new KeyTableData(KeyEvent.KEYCODE_G, KeyEvent.META_SHIFT_ON),	// 71
+		new KeyTableData(KeyEvent.KEYCODE_H, KeyEvent.META_SHIFT_ON),	// 72
+		new KeyTableData(KeyEvent.KEYCODE_I, KeyEvent.META_SHIFT_ON),	// 73
+		new KeyTableData(KeyEvent.KEYCODE_J, KeyEvent.META_SHIFT_ON),	// 74
+		new KeyTableData(KeyEvent.KEYCODE_K, KeyEvent.META_SHIFT_ON),	// 75
+		new KeyTableData(KeyEvent.KEYCODE_L, KeyEvent.META_SHIFT_ON),	// 76
+		new KeyTableData(KeyEvent.KEYCODE_M, KeyEvent.META_SHIFT_ON),	// 77
+		new KeyTableData(KeyEvent.KEYCODE_N, KeyEvent.META_SHIFT_ON),	// 78
+		new KeyTableData(KeyEvent.KEYCODE_O, KeyEvent.META_SHIFT_ON),	// 79
+		new KeyTableData(KeyEvent.KEYCODE_P, KeyEvent.META_SHIFT_ON),	// 80
+		new KeyTableData(KeyEvent.KEYCODE_Q, KeyEvent.META_SHIFT_ON),	// 81
+		new KeyTableData(KeyEvent.KEYCODE_R, KeyEvent.META_SHIFT_ON),	// 82
+		new KeyTableData(KeyEvent.KEYCODE_S, KeyEvent.META_SHIFT_ON),	// 83
+		new KeyTableData(KeyEvent.KEYCODE_T, KeyEvent.META_SHIFT_ON),	// 84
+		new KeyTableData(KeyEvent.KEYCODE_U, KeyEvent.META_SHIFT_ON),	// 85
+		new KeyTableData(KeyEvent.KEYCODE_V, KeyEvent.META_SHIFT_ON),	// 86
+		new KeyTableData(KeyEvent.KEYCODE_W, KeyEvent.META_SHIFT_ON),	// 87
+		new KeyTableData(KeyEvent.KEYCODE_X, KeyEvent.META_SHIFT_ON),	// 88
+		new KeyTableData(KeyEvent.KEYCODE_Y, KeyEvent.META_SHIFT_ON),	// 89
+		new KeyTableData(KeyEvent.KEYCODE_Z, KeyEvent.META_SHIFT_ON),	// 90
+		new KeyTableData(KeyEvent.KEYCODE_LEFT_BRACKET, 0),	// 91
+		new KeyTableData(KeyEvent.KEYCODE_BACKSLASH, 0),	// 92
+		new KeyTableData(KeyEvent.KEYCODE_RIGHT_BRACKET, 0),	// 93
+		new KeyTableData(KeyEvent.KEYCODE_6, KeyEvent.META_SHIFT_ON),	// 94
+		new KeyTableData(KeyEvent.KEYCODE_MINUS, KeyEvent.META_SHIFT_ON),	// 95
+		new KeyTableData(KeyEvent.KEYCODE_GRAVE, 0),	// 96
+		new KeyTableData(KeyEvent.KEYCODE_A, 0),	// 97
+		new KeyTableData(KeyEvent.KEYCODE_B, 0),	// 98
+		new KeyTableData(KeyEvent.KEYCODE_C, 0),	// 99
+		new KeyTableData(KeyEvent.KEYCODE_D, 0),	// 100
+		new KeyTableData(KeyEvent.KEYCODE_E, 0),	// 101
+		new KeyTableData(KeyEvent.KEYCODE_F, 0),	// 102
+		new KeyTableData(KeyEvent.KEYCODE_G, 0),	// 103
+		new KeyTableData(KeyEvent.KEYCODE_H, 0),	// 104
+		new KeyTableData(KeyEvent.KEYCODE_I, 0),	// 105
+		new KeyTableData(KeyEvent.KEYCODE_J, 0),	// 106
+		new KeyTableData(KeyEvent.KEYCODE_K, 0),	// 107
+		new KeyTableData(KeyEvent.KEYCODE_L, 0),	// 108
+		new KeyTableData(KeyEvent.KEYCODE_M, 0),	// 109
+		new KeyTableData(KeyEvent.KEYCODE_N, 0),	// 110
+		new KeyTableData(KeyEvent.KEYCODE_O, 0),	// 111
+		new KeyTableData(KeyEvent.KEYCODE_P, 0),	// 112
+		new KeyTableData(KeyEvent.KEYCODE_Q, 0),	// 113
+		new KeyTableData(KeyEvent.KEYCODE_R, 0),	// 114
+		new KeyTableData(KeyEvent.KEYCODE_S, 0),	// 115
+		new KeyTableData(KeyEvent.KEYCODE_T, 0),	// 116
+		new KeyTableData(KeyEvent.KEYCODE_U, 0),	// 117
+		new KeyTableData(KeyEvent.KEYCODE_V, 0),	// 118
+		new KeyTableData(KeyEvent.KEYCODE_W, 0),	// 119
+		new KeyTableData(KeyEvent.KEYCODE_X, 0),	// 120
+		new KeyTableData(KeyEvent.KEYCODE_Y, 0),	// 121
+		new KeyTableData(KeyEvent.KEYCODE_Z, 0),	// 122
+		new KeyTableData(KeyEvent.KEYCODE_LEFT_BRACKET, KeyEvent.META_SHIFT_ON),	// 123
+		new KeyTableData(KeyEvent.KEYCODE_BACKSLASH, KeyEvent.META_SHIFT_ON),	// 124
+		new KeyTableData(KeyEvent.KEYCODE_RIGHT_BRACKET, KeyEvent.META_SHIFT_ON),	// 125
+		new KeyTableData(KeyEvent.KEYCODE_GRAVE, KeyEvent.META_SHIFT_ON),	// 126
+		new KeyTableData(KeyEvent.KEYCODE_UNKNOWN, 0)	// 127
+	};
+
+	public boolean SendKeyEvent(boolean bDown, int keycode)
+	{
+		if (keycode < 0 && keycode > 127)
+		{
+			return false;
+		}
+
+		final int actionType = bDown ? KeyEvent.ACTION_DOWN : KeyEvent.ACTION_UP;
+		final long actionTime = SystemClock.uptimeMillis();
+		final int actionCode = KeyTableList[keycode].key;
+		final int actionMeta = KeyTableList[keycode].meta;
+//		GameActivity.Log.debug("SendTouchEvent(event=" + (bDown ? "Down" : "Up") + ", key=" + keycode + ", code=" + actionCode + ", meta=" + actionMeta);
+		GameActivity._activity.runOnUiThread(new Runnable()
+		{
+			@Override
+			public void run()
+			{
+				long eventTime = SystemClock.uptimeMillis();
+				webView.dispatchKeyEvent(new KeyEvent(actionTime, eventTime, actionType, actionCode, 0, actionMeta));
+			}
+		});
+		return true;
+	}
+
 	// called from C++ paint event
 	public void Update(final int x, final int y, final int width, final int height)
 	{
diff --git a/Engine/Source/Runtime/WebBrowser/Private/Android/AndroidJavaWebBrowser.cpp b/Engine/Source/Runtime/WebBrowser/Private/Android/AndroidJavaWebBrowser.cpp
--- a/Engine/Source/Runtime/WebBrowser/Private/Android/AndroidJavaWebBrowser.cpp
+++ b/Engine/Source/Runtime/WebBrowser/Private/Android/AndroidJavaWebBrowser.cpp
@@ -50,6 +50,7 @@ FJavaAndroidWebBrowser::FJavaAndroidWebBrowser(bool swizzlePixels, bool vulkanRe
 	, CloseMethod(GetClassMethod("Close", "()V"))
 	, GoBackOrForwardMethod(GetClassMethod("GoBackOrForward", "(I)V"))
 	, SendTouchEventMethod(GetClassMethod("SendTouchEvent", "(IFF)V"))
+	, SendKeyEventMethod(GetClassMethod("SendKeyEvent", "(ZI)Z"))
 	, SetAndroid3DBrowserMethod(GetClassMethod("SetAndroid3DBrowser", "(Z)V"))
 	, SetVisibilityMethod(GetClassMethod("SetVisibility", "(Z)V"))
 {
@@ -293,6 +294,16 @@ void FJavaAndroidWebBrowser::SendTouchMove(float x, float y)
 	CallMethod<void>(SendTouchEventMethod, MOTIONEVENT_ACTION_MOVE, x, y);
 }
 
+bool FJavaAndroidWebBrowser::SendKeyDown(int32 KeyCode)
+{
+	return CallMethod<bool>(SendKeyEventMethod, true, KeyCode);
+}
+
+bool FJavaAndroidWebBrowser::SendKeyUp(int32 KeyCode)
+{
+	return CallMethod<bool>(SendKeyEventMethod, false, KeyCode);
+}
+
 void FJavaAndroidWebBrowser::SetAndroid3DBrowser(bool InIsAndroid3DBrowser)
 {
 	CallMethod<void>(SetAndroid3DBrowserMethod, InIsAndroid3DBrowser);
diff --git a/Engine/Source/Runtime/WebBrowser/Private/Android/AndroidJavaWebBrowser.h b/Engine/Source/Runtime/WebBrowser/Private/Android/AndroidJavaWebBrowser.h
--- a/Engine/Source/Runtime/WebBrowser/Private/Android/AndroidJavaWebBrowser.h
+++ b/Engine/Source/Runtime/WebBrowser/Private/Android/AndroidJavaWebBrowser.h
@@ -34,6 +34,8 @@ public:
 	void SendTouchDown(float x, float y);
 	void SendTouchUp(float x, float y);
 	void SendTouchMove(float x, float y);
+	bool SendKeyDown(int32 KeyCode);
+	bool SendKeyUp(int32 KeyCode);
 	void SetAndroid3DBrowser(bool InIsAndroid3DBrowser);
 	void SetVisibility(bool InIsVisible);
 	void Update(const int posX, const int posY, const int sizeX, const int sizeY);
@@ -55,6 +57,7 @@ private:
 	FJavaClassMethod CloseMethod;
 	FJavaClassMethod GoBackOrForwardMethod;
 	FJavaClassMethod SendTouchEventMethod;
+	FJavaClassMethod SendKeyEventMethod;
 	FJavaClassMethod SetAndroid3DBrowserMethod;
 	FJavaClassMethod SetVisibilityMethod;
 
diff --git a/Engine/Source/Runtime/WebBrowser/Private/Android/AndroidWebBrowserWidget.cpp b/Engine/Source/Runtime/WebBrowser/Private/Android/AndroidWebBrowserWidget.cpp
--- a/Engine/Source/Runtime/WebBrowser/Private/Android/AndroidWebBrowserWidget.cpp
+++ b/Engine/Source/Runtime/WebBrowser/Private/Android/AndroidWebBrowserWidget.cpp
@@ -599,6 +599,29 @@ FReply SAndroidWebBrowserWidget::OnMouseMove(const FGeometry& MyGeometry, const
 	return Reply;
 }
 
+FReply SAndroidWebBrowserWidget::OnKeyDown(const FGeometry& MyGeometry, const FKeyEvent& InKeyEvent)
+{
+//	FPlatformMisc::LowLevelOutputDebugStringf(TEXT("SAndroidWebBrowserWidget::OnKeyDown: %d"), InKeyEvent.GetCharacter());
+	return JavaWebBrowser->SendKeyDown(InKeyEvent.GetCharacter()) ? FReply::Handled() : FReply::Unhandled();
+}
+
+FReply SAndroidWebBrowserWidget::OnKeyUp(const FGeometry& MyGeometry, const FKeyEvent& InKeyEvent)
+{
+//	FPlatformMisc::LowLevelOutputDebugStringf(TEXT("SAndroidWebBrowserWidget::OnKeyUp: %d"), InKeyEvent.GetCharacter());
+	return JavaWebBrowser->SendKeyUp(InKeyEvent.GetCharacter()) ? FReply::Handled() : FReply::Unhandled();
+}
+
+FReply SAndroidWebBrowserWidget::OnKeyChar(const FGeometry& MyGeometry, const FCharacterEvent& InCharacterEvent)
+{
+//	FPlatformMisc::LowLevelOutputDebugStringf(TEXT("SAndroidWebBrowserWidget::OnKeyChar: %d"), (int32)InCharacterEvent.GetCharacter());
+	if (JavaWebBrowser->SendKeyDown(InCharacterEvent.GetCharacter()))
+	{
+		JavaWebBrowser->SendKeyUp(InCharacterEvent.GetCharacter());
+		return FReply::Handled();
+	}
+	return FReply::Unhandled();
+}
+
 void SAndroidWebBrowserWidget::SetWebBrowserVisibility(bool InIsVisible)
 {
 	JavaWebBrowser->SetVisibility(InIsVisible);
diff --git a/Engine/Source/Runtime/WebBrowser/Private/Android/AndroidWebBrowserWidget.h b/Engine/Source/Runtime/WebBrowser/Private/Android/AndroidWebBrowserWidget.h
--- a/Engine/Source/Runtime/WebBrowser/Private/Android/AndroidWebBrowserWidget.h
+++ b/Engine/Source/Runtime/WebBrowser/Private/Android/AndroidWebBrowserWidget.h
@@ -25,7 +25,7 @@
 class UMaterialExpressionTextureSample;
 class FWebBrowserTextureSamplePool;
 
-class SAndroidWebBrowserWidget : public SLeafWidget
+class SAndroidWebBrowserWidget : public SViewport
 {
 	SLATE_BEGIN_ARGS(SAndroidWebBrowserWidget)
 		: _InitialURL("about:blank")
@@ -67,6 +67,10 @@ public:
 	void SendTouchUp(FVector2D Position);
 	void SendTouchMove(FVector2D Position);
 
+	virtual FReply OnKeyDown(const FGeometry& MyGeometry, const FKeyEvent& InKeyEvent) override;
+	virtual FReply OnKeyUp(const FGeometry& MyGeometry, const FKeyEvent& InKeyEvent) override;
+	virtual FReply OnKeyChar(const FGeometry& MyGeometry, const FCharacterEvent& InCharacterEvent) override;
+
 	// WebViewClient callbacks
 
 	jbyteArray HandleShouldInterceptRequest(jstring JUrl);
diff --git a/Engine/Source/Runtime/WebBrowser/Private/Android/AndroidWebBrowserWindow.cpp b/Engine/Source/Runtime/WebBrowser/Private/Android/AndroidWebBrowserWindow.cpp
--- a/Engine/Source/Runtime/WebBrowser/Private/Android/AndroidWebBrowserWindow.cpp
+++ b/Engine/Source/Runtime/WebBrowser/Private/Android/AndroidWebBrowserWindow.cpp
@@ -113,17 +113,17 @@ FString FAndroidWebBrowserWindow::GetUrl() const
 
 bool FAndroidWebBrowserWindow::OnKeyDown(const FKeyEvent& InKeyEvent)
 {
-	return false;
+	return BrowserWidget->OnKeyDown(FGeometry(), InKeyEvent).IsEventHandled();
 }
 
 bool FAndroidWebBrowserWindow::OnKeyUp(const FKeyEvent& InKeyEvent)
 {
-	return false;
+	return BrowserWidget->OnKeyUp(FGeometry(), InKeyEvent).IsEventHandled();
 }
 
 bool FAndroidWebBrowserWindow::OnKeyChar(const FCharacterEvent& InCharacterEvent)
 {
-	return false;
+	return BrowserWidget->OnKeyChar(FGeometry(), InCharacterEvent).IsEventHandled();
 }
 
 FVector2D FAndroidWebBrowserWindow::ConvertMouseEventToLocal(const FGeometry& MyGeometry, const FPointerEvent& MouseEvent, bool bIsPopup)
