Reduce PSO compiler service priority.
Fixes hitching while external PSO compiles are running.

diff --git a/Engine/Build/Android/Java/src/com/epicgames/unreal/psoservices/PSOProgramService.java b/Engine/Build/Android/Java/src/com/epicgames/unreal/psoservices/PSOProgramService.java
--- a/Engine/Build/Android/Java/src/com/epicgames/unreal/psoservices/PSOProgramService.java
+++ b/Engine/Build/Android/Java/src/com/epicgames/unreal/psoservices/PSOProgramService.java
@@ -329,6 +329,7 @@ public class PSOProgramService extends Service implements Logger.ILoggerCallback
 		super.onCreate();
 		try
 		{
+			android.os.Process.setThreadPriority(android.os.Process.THREAD_PRIORITY_BACKGROUND);
 			PrepareLogger();
 		} 
 		catch (IOException e)
