Adding another way to get device metrics as Display#getRealSize() is deprecated in API level 31

[RN] Adding another way to get device metrics for API level 31 and higher

diff --git a/Engine/Build/Android/Java/src/com/epicgames/unreal/GameActivity.java.template b/Engine/Build/Android/Java/src/com/epicgames/unreal/GameActivity.java.template
--- a/Engine/Build/Android/Java/src/com/epicgames/unreal/GameActivity.java.template
+++ b/Engine/Build/Android/Java/src/com/epicgames/unreal/GameActivity.java.template
@@ -5711,10 +5711,17 @@ public class GameActivity extends $${gameActivitySuperClass}$$ implements Surfac
 		if (key.equals("unreal.displaymetrics.dpi"))
 		{
 			DisplayMetrics metrics = new DisplayMetrics();
+			if (android.os.Build.VERSION.SDK_INT >= 31)
+			{
+				metrics = getResources().getDisplayMetrics();
+			}
+			else
 			if (android.os.Build.VERSION.SDK_INT >= 17)
 			{
 				getWindowManager().getDefaultDisplay().getRealMetrics(metrics);
-			} else {
+			}
+			else
+			{
 				// note: not available so get what we can
 				getWindowManager().getDefaultDisplay().getMetrics(metrics);
 			}
