Fix a crash with the vulkan PSO compile service which occurs on subsequent invocations.

diff --git a/Engine/Build/Android/Prebuilt/PSOService/Source/PSOService.cpp b/Engine/Build/Android/Prebuilt/PSOService/Source/PSOService.cpp
--- a/Engine/Build/Android/Prebuilt/PSOService/Source/PSOService.cpp
+++ b/Engine/Build/Android/Prebuilt/PSOService/Source/PSOService.cpp
@@ -204,6 +204,8 @@ public:
 
 	void ShutDownDevice()
 	{
+		bInitialized = false;
+
 		if (Device == VK_NULL_HANDLE)
 		{
 			return;
