Change Memstats printing function.
Allow config to change printing interval

diff --git a/Engine/Build/Android/Java/src/com/epicgames/unreal/GameActivity.java.template b/Engine/Build/Android/Java/src/com/epicgames/unreal/GameActivity.java.template
--- a/Engine/Build/Android/Java/src/com/epicgames/unreal/GameActivity.java.template
+++ b/Engine/Build/Android/Java/src/com/epicgames/unreal/GameActivity.java.template
@@ -507,6 +507,10 @@ public class GameActivity extends $${gameActivitySuperClass}$$ implements Surfac
 	private boolean localNotificationAppLaunched = false;
 	private String	localNotificationLaunchActivationEvent = "";
 	private int		localNotificationLaunchFireDate = 0;
+
+	/** Memstats logging settings */
+	private int memstatLoggingInterval = 30000;
+	private int memstatRSSRequiredDiff = 2048;
 	
 	enum EAlertDialogType
 	{
@@ -2255,6 +2259,30 @@ public class GameActivity extends $${gameActivitySuperClass}$$ implements Surfac
 			}
 		}
 
+		try 
+		{
+			if (variables.containsKey("MemstatLoggingInterval"))
+			{
+				memstatLoggingInterval = Integer.decode(variables.get("MemstatLoggingInterval"));
+			}
+		} 
+		catch (Exception e)
+		{
+			Log.debug("Failed to read memstats: MemstatLoggingInterval");
+		}
+		try 
+		{
+			if (variables.containsKey("MemstatRSSRequiredDiff"))
+			{
+				memstatRSSRequiredDiff = Integer.decode(variables.get("MemstatRSSRequiredDiff"));
+			}
+		} 
+		catch (Exception e)
+		{
+			Log.debug("Failed to read memstats: MemstatRSSRequiredDiff");
+		}
+		
+
 		// notify AndroidProcess of the core masks and affinity enable state
 		if (variables.containsKey("unsupportedDevice"))
 		{
@@ -2479,10 +2507,12 @@ public class GameActivity extends $${gameActivitySuperClass}$$ implements Surfac
 			Logger.SuppressLogs();
 		}
 
-		// update memory stats every 10 seconds
+		// update memory stats every "memstatLoggingInterval" seconds
+		// only print if VmRSS diff is >= "memstatRSSRequiredDiff" bytes 
 		memoryRunnable = new Runnable() {
 			boolean showPSSInfo = false;
 			int LastProcessMemoryValue = 0;
+			int LastRSSMemoryValue = 0;
 			@Override
 			public void run()
 			{
@@ -2506,26 +2536,36 @@ public class GameActivity extends $${gameActivitySuperClass}$$ implements Surfac
 					}
 				}
 
-				Log.debug("--- memstats:");
 				String[] statusLines = FindLineFromFile("/proc/self/status",new String[] {"VmSize:","VmHWM:", "VmRSS:", "VmData:", "VmSwap:"});
+				String rssSubString = statusLines[2].substring(0, statusLines[2].lastIndexOf(" "));
+				int rssMemory = Integer.parseInt(rssSubString.substring(rssSubString.lastIndexOf(" ")+1, rssSubString.length()));
+				int rssDiffed = Math.abs(rssMemory - LastRSSMemoryValue);
+
 				String ProcessMemInfoLine = "Used memory: ";
-				if( LastProcessMemoryValue != ProcessMemory)
-				{
-					ProcessMemInfoLine+=" (updated): ";
-					LastProcessMemoryValue = ProcessMemory;
-				}
-				Log.debug("memstat: "+ProcessMemInfoLine + ProcessMemory);
-				for(String memstatusline : statusLines)
-				{
-					Log.debug("memstat: "+memstatusline);
-				}
-				if( showPSSInfo )
+				if(rssDiffed >= memstatRSSRequiredDiff)
 				{
-					String[] smapsLines = FindLineFromFile("/proc/self/smaps_rollup",new String[] {"SwapPss:","Pss:"});
-				    for(String memstatusline : smapsLines)
-				    {
-					    Log.debug("memstat: "+memstatusline);
-				    }
+					String memLog = "--- memstats: \n";
+					LastRSSMemoryValue = rssMemory;
+					if( LastProcessMemoryValue != ProcessMemory)
+					{
+						ProcessMemInfoLine+=" (updated): ";
+						LastProcessMemoryValue = ProcessMemory;
+					}
+					memLog += "memstat: "+ProcessMemInfoLine + ProcessMemory + "\n";
+					for(String memstatusline : statusLines)
+					{
+						memLog += "memstat: "+memstatusline + "\n";
+					}
+				
+					if( showPSSInfo )
+					{
+						String[] smapsLines = FindLineFromFile("/proc/self/smaps_rollup",new String[] {"SwapPss:","Pss:"});
+				    	for(String memstatusline : smapsLines)
+				    	{
+							memLog += "memstat: "+memstatusline + "\n";
+				    	}
+					}
+					Log.debug(memLog);
 				}
 				synchronized(_activity)
 				{
@@ -2534,9 +2574,9 @@ public class GameActivity extends $${gameActivitySuperClass}$$ implements Surfac
 						_activity.UsedMemory = ProcessMemory;
 					}
 					
-					if (_activity.memoryHandler != null)
+					if (_activity.memoryHandler != null && memstatLoggingInterval > 0)
 					{
-						_activity.memoryHandler.postDelayed(this, 10000);
+						_activity.memoryHandler.postDelayed(this, memstatLoggingInterval);
 					}
 				}
 			}
@@ -3418,10 +3458,10 @@ public class GameActivity extends $${gameActivitySuperClass}$$ implements Surfac
 		super.onResume();
 		bIsActivityPaused.set(false);
 
-		// start memory reporter timer (runs every 10 seconds)
+		// start memory reporter timer (runs every "memstatLoggingInterval" seconds)
 		synchronized(this)
 		{
-			if (memoryHandler == null && memoryRunnable != null)
+			if (memoryHandler == null && memoryRunnable != null && memstatLoggingInterval > 0)
 			{
 				Log.debug("onResume: start memory reporter runnable");
 				memoryHandler = new Handler(memoryHandlerThread.getLooper());
