Fix up support for Google Play Services (remove CVE-[REDACTED_PHONE] issue)

#android

diff --git a/Engine/Build/Android/Java/aar-imports-ant.txt b/Engine/Build/Android/Java/aar-imports-ant.txt
deleted file mode 100644
diff --git a/Engine/Build/Android/Java/aar-imports.txt b/Engine/Build/Android/Java/aar-imports.txt
--- a/Engine/Build/Android/Java/aar-imports.txt
+++ b/Engine/Build/Android/Java/aar-imports.txt
@@ -1,6 +1,6 @@
 repositories $(ANDROID_HOME)/extras
 repositories $(ENGINEDIR)/Source/ThirdParty/Android/extras
-com.google.android.gms,play-services-auth,11.8.0
-com.google.android.gms,play-services-games,23.0.0
-com.google.android.gms,play-services-nearby,11.8.0
-com.google.android.gms,play-services-plus,11.8.0
+com.google.android.gms,play-services-auth,18.1.0
+com.google.android.gms,play-services-nearby,18.1.0
+com.google.android.gms,play-services-games,18.0.1
+com.google.android.gms,play-services-plus,17.0.0
