Allow Instant notifications to be sent
* If TargetTime is set to 0 it will send it as soon as posible

diff --git a/Engine/Build/Android/Java/src/com/epicgames/unreal/GameActivity.java.template b/Engine/Build/Android/Java/src/com/epicgames/unreal/GameActivity.java.template
--- a/Engine/Build/Android/Java/src/com/epicgames/unreal/GameActivity.java.template
+++ b/Engine/Build/Android/Java/src/com/epicgames/unreal/GameActivity.java.template
@@ -5254,10 +5254,21 @@ public class GameActivity extends $${gameActivitySuperClass}$$ implements Surfac
 
 		long msDiff = targetDate.getTime() - currentDate.getTime();
 
+		// If you set the targetDateTime to 0, just send the notification as soon as possible
 		if(msDiff < 0)
 		{
-			LocalNotificationRemoveID(context, notificationID);
-			return false;
+			if (targetDate.getYear() == -1899)
+			{
+				// AlarmManager delivery is inexact: So instead of using the alarmManager to send instant notifications, just sending them right away
+				LocalNotificationReceiver directSend = new LocalNotificationReceiver();
+				directSend.onReceive(context, notificationIntent);
+				return false;
+			}
+			else
+			{
+				LocalNotificationRemoveID(context, notificationID);
+				return false;
+			}
 		}
 
 		long futureTimeInMillis = SystemClock.elapsedRealtime() + msDiff;//Calculate the time to run the callback
diff --git a/Engine/Source/Runtime/Online/BackgroundHTTP/Private/BackgroundHttpNotificationObject.cpp b/Engine/Source/Runtime/Online/BackgroundHTTP/Private/BackgroundHttpNotificationObject.cpp
--- a/Engine/Source/Runtime/Online/BackgroundHTTP/Private/BackgroundHttpNotificationObject.cpp
+++ b/Engine/Source/Runtime/Online/BackgroundHTTP/Private/BackgroundHttpNotificationObject.cpp
@@ -82,11 +82,8 @@ FBackgroundHttpNotificationObject::~FBackgroundHttpNotificationObject()
 	{
 		if (!bNotifyOnlyOnFullSuccess || (NumFailedDownloads == 0))
 		{
-			//Used to be 15s, probably for reasons to do with Ios. 
-			//make a notification 1 second from now
-			FDateTime TargetTime = FDateTime::Now();
-			TargetTime += FTimespan::FromSeconds(1);
-
+			//Setting the datetime to 0 forcing the local notifcation to be sent as soon as possible
+			FDateTime TargetTime = FDateTime();
 			if (nullptr != PlatformNotificationService)
 			{
 				PlatformNotificationService->ScheduleLocalNotificationAtTimeOverrideId(TargetTime, true, NotificationTitle, NotificationBody, NotificationAction, NotificationActivationString, IdOverride);
