Fixing hardware keyboard input not passing through IME

diff --git a/Engine/Build/Android/Java/src/com/epicgames/unreal/GameActivity.java.template b/Engine/Build/Android/Java/src/com/epicgames/unreal/GameActivity.java.template
--- a/Engine/Build/Android/Java/src/com/epicgames/unreal/GameActivity.java.template
+++ b/Engine/Build/Android/Java/src/com/epicgames/unreal/GameActivity.java.template
@@ -6328,7 +6328,7 @@ public class GameActivity extends $${gameActivitySuperClass}$$ implements Surfac
 		    	AndroidThunkJava_HideVirtualKeyboardInput();
 				nativeVirtualKeyboardSendKey(KeyEvent.KEYCODE_BACK);
 		    }
-		    return super.dispatchKeyEvent(event);
+		    return super.onKeyPreIme(event);
 		}
 
 		// listen for cursor position changes &  send message to the UMG widget
