Tentative fix for android safe zone.

diff --git a/Engine/Build/Android/Java/src/com/epicgames/unreal/GameActivity.java.template b/Engine/Build/Android/Java/src/com/epicgames/unreal/GameActivity.java.template
--- a/Engine/Build/Android/Java/src/com/epicgames/unreal/GameActivity.java.template
+++ b/Engine/Build/Android/Java/src/com/epicgames/unreal/GameActivity.java.template
@@ -2940,7 +2940,8 @@ public class GameActivity extends $${gameActivitySuperClass}$$ implements Surfac
 		nativeSetGlobalActivity(UseExternalFilesDir, PublicLogFiles, getFilesDir().getPath(), getExternalFilesDir(null).getPath(), PackageDataInsideApkValue==1, APKPath);
 
 		// tell the engine if this is a portrait app
-		nativeSetWindowInfo(getResources().getConfiguration().orientation == Configuration.ORIENTATION_PORTRAIT, DepthBufferPreference, PropagateAlpha);
+		int orientation = AndroidThunkJava_GetDeviceOrientation();
+		nativeSetWindowInfo(orientation, DepthBufferPreference, PropagateAlpha);
 
 		// get the full language code, like en-US
 		// note: this may need to be Locale.getDefault().getLanguage()
@@ -3383,7 +3384,8 @@ public class GameActivity extends $${gameActivitySuperClass}$$ implements Surfac
 		}
 
 		// invalidate window cache
-		nativeSetWindowInfo(getResources().getConfiguration().orientation == Configuration.ORIENTATION_PORTRAIT, DepthBufferPreference, PropagateAlpha);
+		int orientation = AndroidThunkJava_GetDeviceOrientation();
+		nativeSetWindowInfo(orientation, DepthBufferPreference, PropagateAlpha);
 
 		// only do this on KitKat and above
 		if (ShouldHideUI)
@@ -3878,8 +3880,8 @@ public class GameActivity extends $${gameActivitySuperClass}$$ implements Surfac
 		super.onConfigurationChanged(newConfig);
 		//$${gameActivityonConfigurationChangedAdditions}$$
 		// forward the orientation
-		boolean bPortrait = newConfig.orientation == Configuration.ORIENTATION_PORTRAIT;
-		nativeOnConfigurationChanged(bPortrait);
+		int orientation = AndroidThunkJava_GetDeviceOrientation();
+		nativeOnConfigurationChanged(orientation);
 	}
 
 	public void surfaceChanged(SurfaceHolder holder, int format, int width, int height)
@@ -6953,7 +6955,7 @@ public class GameActivity extends $${gameActivitySuperClass}$$ implements Surfac
 	public native void nativeSetAndroidStartupState(boolean bDebuggerAttached);
 	public native void nativeSetGlobalActivity(boolean bUseExternalFilesDir, boolean bPublicLogFiles, String internalFilePath, String externalFilePath, boolean bOBBInAPK, String APKPath);
 	public native void nativeSetObbFilePaths(String OBBMainFilePath, String OBBPatchFilePath, String OBBOverflow1FilePath, String OBBOverflow2FilePath);
-	public native void nativeSetWindowInfo(boolean bIsPortrait, int DepthBufferPreference, int PropagateAlpha);
+	public native void nativeSetWindowInfo(int orientation, int DepthBufferPreference, int PropagateAlpha);
 	public native void nativeSetObbInfo(String ProjectName, String PackageName, int Version, int PatchVersion, String AppType);
 	public native void nativeSetAndroidVersionInformation( String AndroidVersion, int TargetSDK, String PhoneMake, String PhoneModel, String PhoneBuildNumber, String OSLanguage );
 
@@ -6978,7 +6980,7 @@ public class GameActivity extends $${gameActivitySuperClass}$$ implements Surfac
 	public native void nativeVirtualKeyboardShown(int left, int top, int right, int bottom);
 	public native void nativeVirtualKeyboardVisible(boolean bShown);
 
-	public native void nativeOnConfigurationChanged(boolean bPortrait);
+	public native void nativeOnConfigurationChanged(int orientation);
 		
 	public native void nativeOnInitialDownloadStarted();
 	public native void nativeOnInitialDownloadCompleted();
diff --git a/Engine/Source/Runtime/ApplicationCore/Private/Android/AndroidWindow.cpp b/Engine/Source/Runtime/ApplicationCore/Private/Android/AndroidWindow.cpp
--- a/Engine/Source/Runtime/ApplicationCore/Private/Android/AndroidWindow.cpp
+++ b/Engine/Source/Runtime/ApplicationCore/Private/Android/AndroidWindow.cpp
@@ -129,14 +129,17 @@ void FAndroidWindow::SetOSWindowHandle(void* InWindow)
 
 //This function is declared in the Java-defined class, GameActivity.java: "public native void nativeSetObbInfo(String PackageName, int Version, int PatchVersion);"
 static bool GAndroidIsPortrait = false;
+static int GWindowOrientation = -1;
 static std::atomic<bool> GAndroidSafezoneRequiresUpdate = false;
 static int GAndroidDepthBufferPreference = 0;
 static FVector4 GAndroidPortraitSafezone = FVector4(-1.0f, -1.0f, -1.0f, -1.0f);
 static FVector4 GAndroidLandscapeSafezone = FVector4(-1.0f, -1.0f, -1.0f, -1.0f);
 #if USE_ANDROID_JNI
-JNI_METHOD void Java_com_epicgames_unreal_GameActivity_nativeSetWindowInfo(JNIEnv* jenv, jobject thiz, jboolean bIsPortrait, jint DepthBufferPreference, jint PropagateAlpha)
+JNI_METHOD void Java_com_epicgames_unreal_GameActivity_nativeSetWindowInfo(JNIEnv* jenv, jobject thiz, jint orientation, jint DepthBufferPreference, jint PropagateAlpha)
 {
 	ClearCachedWindowRects();
+	GWindowOrientation = orientation;
+	bool bIsPortrait = GWindowOrientation == EAndroidSurfaceOrientation::ROTATION_0 || GWindowOrientation == EAndroidSurfaceOrientation::ROTATION_180;
 	GAndroidIsPortrait = bIsPortrait == JNI_TRUE;
 	GAndroidDepthBufferPreference = DepthBufferPreference;
 	GAndroidPropagateAlpha = PropagateAlpha;
@@ -557,10 +560,12 @@ void FAndroidWindow::CalculateSurfaceSize(int32_t& SurfaceWidth, int32_t& Surfac
 #endif
 }
 
-bool FAndroidWindow::OnWindowOrientationChanged(bool bIsPortrait)
+bool FAndroidWindow::OnWindowOrientationChanged(int Orientation)
 {
-	if (GAndroidIsPortrait != bIsPortrait)
+	if (GWindowOrientation != Orientation)
 	{
+		GWindowOrientation = Orientation;
+		bool bIsPortrait = GWindowOrientation == EAndroidSurfaceOrientation::ROTATION_0 || GWindowOrientation == EAndroidSurfaceOrientation::ROTATION_180;
 		UE_LOG(LogAndroid, Log, TEXT("Window orientation changed: %s"), bIsPortrait ? TEXT("Portrait") : TEXT("Landscape"));
 		GAndroidIsPortrait = bIsPortrait;
 		return true;
diff --git a/Engine/Source/Runtime/ApplicationCore/Public/Android/AndroidWindow.h b/Engine/Source/Runtime/ApplicationCore/Public/Android/AndroidWindow.h
--- a/Engine/Source/Runtime/ApplicationCore/Public/Android/AndroidWindow.h
+++ b/Engine/Source/Runtime/ApplicationCore/Public/Android/AndroidWindow.h
@@ -43,7 +43,7 @@ public:
 	// When bUseEventThreadWindow == false this uses dimensions cached when the game thread processes android events.
 	// When bUseEventThreadWindow == true this uses dimensions directly from the android event thread, unless called from event thread this requires acquiring GAndroidWindowLock to use.
 	static void CalculateSurfaceSize(int32_t& SurfaceWidth, int32_t& SurfaceHeight, bool bUseEventThreadWindow = false);
-	static bool OnWindowOrientationChanged(bool bIsPortrait);
+	static bool OnWindowOrientationChanged(int Orientation);
 
 	static int32 GetDepthBufferPreference();
 	
diff --git a/Engine/Source/Runtime/Core/Private/Android/AndroidPlatformMisc.cpp b/Engine/Source/Runtime/Core/Private/Android/AndroidPlatformMisc.cpp
--- a/Engine/Source/Runtime/Core/Private/Android/AndroidPlatformMisc.cpp
+++ b/Engine/Source/Runtime/Core/Private/Android/AndroidPlatformMisc.cpp
@@ -636,10 +636,10 @@ void FAndroidMisc::UpdateDeviceOrientation()
 			const int Orientation = JEnv->CallIntMethod(AndroidJavaEnv::GetGameActivityThis(), getOrientationMethod);
 			switch (Orientation)
 			{
-			case 0: DeviceOrientation = EDeviceScreenOrientation::Portrait;             break;
-			case 1: DeviceOrientation = EDeviceScreenOrientation::LandscapeLeft;        break;
-			case 2: DeviceOrientation = EDeviceScreenOrientation::PortraitUpsideDown;   break;
-			case 3: DeviceOrientation = EDeviceScreenOrientation::LandscapeRight;       break;
+			case EAndroidSurfaceOrientation::ROTATION_0:	DeviceOrientation = EDeviceScreenOrientation::Portrait;             break;
+			case EAndroidSurfaceOrientation::ROTATION_90:	DeviceOrientation = EDeviceScreenOrientation::LandscapeLeft;        break;
+			case EAndroidSurfaceOrientation::ROTATION_180:	DeviceOrientation = EDeviceScreenOrientation::PortraitUpsideDown;   break;
+			case EAndroidSurfaceOrientation::ROTATION_270:	DeviceOrientation = EDeviceScreenOrientation::LandscapeRight;       break;
 			}
 		}
 	}
diff --git a/Engine/Source/Runtime/Core/Public/Android/AndroidPlatformMisc.h b/Engine/Source/Runtime/Core/Public/Android/AndroidPlatformMisc.h
--- a/Engine/Source/Runtime/Core/Public/Android/AndroidPlatformMisc.h
+++ b/Engine/Source/Runtime/Core/Public/Android/AndroidPlatformMisc.h
@@ -23,6 +23,18 @@ class TFunction;
 
 enum class ECrashContextType;
 
+//The android surface orientations defined in C++ map to the values in https://developer.android.com/reference/android/view/Surface#ROTATION_0
+namespace EAndroidSurfaceOrientation
+{
+	enum Type
+	{
+		ROTATION_0,
+		ROTATION_90,
+		ROTATION_180,
+		ROTATION_270,
+	};
+};
+
 /**
  * Android implementation of the misc OS functions
  */
diff --git a/Engine/Source/Runtime/Launch/Private/Android/LaunchAndroid.cpp b/Engine/Source/Runtime/Launch/Private/Android/LaunchAndroid.cpp
--- a/Engine/Source/Runtime/Launch/Private/Android/LaunchAndroid.cpp
+++ b/Engine/Source/Runtime/Launch/Private/Android/LaunchAndroid.cpp
@@ -1642,13 +1642,11 @@ JNI_METHOD jboolean Java_com_epicgames_unreal_GameActivity_nativeSupportsNEON(JN
 }
 
 //This function is declared in the Java-defined class, GameActivity.java: "public native void nativeOnConfigurationChanged(boolean bPortrait);
-JNI_METHOD void Java_com_epicgames_unreal_GameActivity_nativeOnConfigurationChanged(JNIEnv* jenv, jobject thiz, jboolean bPortrait)
+JNI_METHOD void Java_com_epicgames_unreal_GameActivity_nativeOnConfigurationChanged(JNIEnv* jenv, jobject thiz, jint orientation)
 {
-	bool bChangedToPortrait = bPortrait == JNI_TRUE;
-
 	// enqueue a window changed event if orientation changed, 
 	// note that the HW window handle does not necessarily change.
-	if (FAndroidWindow::OnWindowOrientationChanged(bChangedToPortrait))
+	if (FAndroidWindow::OnWindowOrientationChanged(orientation))
 	{	
 		// Enqueue an event to trigger gamethread to update the orientation:
 		FAppEventManager::GetInstance()->EnqueueAppEvent(APP_EVENT_STATE_WINDOW_CHANGED);
