Add additional Razer Kishi gamepad

#android

diff --git a/Engine/Build/Android/Java/src/com/epicgames/unreal/GameActivity.java.template b/Engine/Build/Android/Java/src/com/epicgames/unreal/GameActivity.java.template
--- a/Engine/Build/Android/Java/src/com/epicgames/unreal/GameActivity.java.template
+++ b/Engine/Build/Android/Java/src/com/epicgames/unreal/GameActivity.java.template
@@ -5391,6 +5391,7 @@ public class GameActivity extends $${gameActivitySuperClass}$$ implements Surfac
 		new DeviceInfoData(0x054c, 0x0ce6, "PS5 Wireless Controller"),
 		new DeviceInfoData(0x05ac, 0x056a, "glap QXPGP001"),
 		new DeviceInfoData(0x0483, 0x5750, "STMicroelectronics Lenovo GamePad"),
+		new DeviceInfoData(0x1532, 0x0037, "Razer Kishi V2 Pro XBox360"),
 		new DeviceInfoData(0x1532, 0x0705, "Razer Razer Claire T1 Wired"),
 		new DeviceInfoData(0x1532, 0x0712, "Razer Kishi V2"),
 		new DeviceInfoData(0x1532, 0x0718, "Razer Kishi V2 Pro XBox360")
