Removing emoji filtering on Android
Originally 25582853

diff --git a/Engine/Build/Android/Java/src/com/epicgames/unreal/GameActivity.java.template b/Engine/Build/Android/Java/src/com/epicgames/unreal/GameActivity.java.template
--- a/Engine/Build/Android/Java/src/com/epicgames/unreal/GameActivity.java.template
+++ b/Engine/Build/Android/Java/src/com/epicgames/unreal/GameActivity.java.template
@@ -6275,69 +6275,18 @@ public class GameActivity extends $${gameActivitySuperClass}$$ implements Surfac
 		public VirtualKeyboardInput(Context context, AttributeSet attrs, int defStyle) 
 		{
 			super(context, attrs, defStyle);
-	        init();
 		}
 
 		public VirtualKeyboardInput(Context context, AttributeSet attrs) 
 		{
 			super(context, attrs);
-	        init();
 		}
 
 		public VirtualKeyboardInput(Context context) 
 		{
 			super(context);
-	        init();
 		}
 
-		private void init() 
-		{
-			if (emojiExcludeFilter == null)
-			{
-				emojiExcludeFilter = new EmojiExcludeFilter();
-			}
-			//setFilters(new InputFilter[]{emojiExcludeFilter});
-		}
-
-		@Override
-		public void setFilters(InputFilter[] filters) 
-		{
-			if (filters.length != 0 && emojiExcludeFilter != null) 
-			{ //if length == 0 it will here return when init() is called
-				boolean add = true;
-				for (InputFilter inputFilter : filters) 
-				{
-					if (inputFilter == emojiExcludeFilter) 
-					{
-						add = false;
-						break;
-					}
-				}
-				if (add)
-				{
-					filters = Arrays.copyOf(filters, filters.length + 1);
-					filters[filters.length - 1] = emojiExcludeFilter;
-				}
-			}
-			super.setFilters(filters);
-		}
-		    
-		private EmojiExcludeFilter emojiExcludeFilter;
-
-	    private class EmojiExcludeFilter implements InputFilter 
-	    {
-	        @Override
-	        public CharSequence filter(CharSequence source, int start, int end, Spanned dest, int dstart, int dend) {
-	            for (int i = start; i < end; i++) {
-	                int type = Character.getType(source.charAt(i));
-	                if (type == Character.SURROGATE || type == Character.OTHER_SYMBOL) {
-	                    return "";
-	                }
-	            }
-	            return null;
-	        }
-	    }
-
 		//Override BACK key to hide the virtual keyboard
 		@Override 
 		public boolean onKeyPreIme(int keyCode, KeyEvent event) 
