Fix for [Mobile] [Luna Controller] The left and right triggers are not functional
   -> Added a name check for Luna controllers, to enable using the threshold to consider pulling triggers as "button presses".
   -> Added controller info to GameActivity

#rnx

diff --git a/Engine/Build/Android/Java/src/com/epicgames/unreal/GameActivity.java.template b/Engine/Build/Android/Java/src/com/epicgames/unreal/GameActivity.java.template
--- a/Engine/Build/Android/Java/src/com/epicgames/unreal/GameActivity.java.template
+++ b/Engine/Build/Android/Java/src/com/epicgames/unreal/GameActivity.java.template
@@ -5391,7 +5391,8 @@ public class GameActivity extends $${gameActivitySuperClass}$$ implements Surfac
 		new DeviceInfoData(0x1532, 0x0037, "Razer Kishi V2 Pro XBox360"),
 		new DeviceInfoData(0x1532, 0x0705, "Razer Razer Claire T1 Wired"),
 		new DeviceInfoData(0x1532, 0x0712, "Razer Kishi V2"),
-		new DeviceInfoData(0x1532, 0x0718, "Razer Kishi V2 Pro XBox360")
+		new DeviceInfoData(0x1532, 0x0718, "Razer Kishi V2 Pro XBox360"),
+		new DeviceInfoData(0x0171, 0x0419, "Luna Gamepad")
 	};
 
 	public class InputDeviceInfo {
diff --git a/Engine/Source/Runtime/ApplicationCore/Private/Android/AndroidInputInterface.cpp b/Engine/Source/Runtime/ApplicationCore/Private/Android/AndroidInputInterface.cpp
--- a/Engine/Source/Runtime/ApplicationCore/Private/Android/AndroidInputInterface.cpp
+++ b/Engine/Source/Runtime/ApplicationCore/Private/Android/AndroidInputInterface.cpp
@@ -1160,6 +1160,10 @@ void FAndroidInputInterface::SendControllerEvents()
 								CurrentDevice.bTriggersUseThresholdForClick = true;
 							}
 						}
+						else if (CurrentDevice.DeviceInfo.Name.StartsWith(TEXT("Luna")))
+						{
+							CurrentDevice.bTriggersUseThresholdForClick = true;
+						}
 
 						IPlatformInputDeviceMapper& DeviceMapper = IPlatformInputDeviceMapper::Get();
 						FPlatformUserId PlatformUserId = PLATFORMUSERID_NONE;
