Fix left/right triggers not functional when connecting over xbox one wired controller

-Added proper product and vendor IDs to the list of known XBox wired controllers

diff --git a/Engine/Build/Android/Java/src/com/epicgames/unreal/GameActivity.java.template b/Engine/Build/Android/Java/src/com/epicgames/unreal/GameActivity.java.template
--- a/Engine/Build/Android/Java/src/com/epicgames/unreal/GameActivity.java.template
+++ b/Engine/Build/Android/Java/src/com/epicgames/unreal/GameActivity.java.template
@@ -5377,6 +5377,7 @@ public class GameActivity extends $${gameActivitySuperClass}$$ implements Surfac
 		new DeviceInfoData(0x0738, 0x5263, "Mad Catz C.T.R.L.R (Smart)"),
 		new DeviceInfoData(0x0738, 0x5266, "Mad Catz C.T.R.L.R"),
 		new DeviceInfoData(0x045e, 0x028e, "Xbox Wired Controller"),
+		new DeviceInfoData(0x045e, 0x02d1, "Xbox Wired Controller"),
 		new DeviceInfoData(0x045e, 0x02e0, "Xbox Wireless Controller"),
 		new DeviceInfoData(0x045e, 0x02fd, "Xbox Wireless Controller"),
 		new DeviceInfoData(0x045e, 0x0b05, "Xbox Elite Wireless Controller"),
