Fix the Safezone updated incorrectly issue, the SafeFrameChangeEvent needs to be triggered after the orientation changed.
Fix the Gyro broken issue.

diff --git a/Engine/Build/Android/Java/src/com/epicgames/unreal/GameActivity.java.template b/Engine/Build/Android/Java/src/com/epicgames/unreal/GameActivity.java.template
--- a/Engine/Build/Android/Java/src/com/epicgames/unreal/GameActivity.java.template
+++ b/Engine/Build/Android/Java/src/com/epicgames/unreal/GameActivity.java.template
@@ -320,6 +320,20 @@ public class GameActivity extends $${gameActivitySuperClass}$$ implements Surfac
 	private static final int CONFIGRULES_STATE_EXECFALSE = 3;
 	private static final int CONFIGRULES_STATE_FINDEND = 4;
 
+	// Same as EDeviceScreenOrientation
+	/** The orientation is not known */
+	private static final int DEVICE_ORIENTATION_UNKNOWN = 0;
+	/** The orientation is portrait with the home button at the bottom */
+	private static final int DEVICE_ORIENTATION_PORTRAIT = 1;
+	/** The orientation is portrait with the home button at the top */
+	private static final int DEVICE_ORIENTATION_PORTRAIT_UPSIDEDOWN = 2;
+	/** The orientation is landscape with the home button at the right side */
+	private static final int DEVICE_ORIENTATION_LANDSCAPE_LEFT = 3;
+	/** The orientation is landscape with the home button at the left side */
+	private static final int DEVICE_ORIENTATION_LANDSCAPE_RIGHT = 4;
+	/** The device screen orientation */
+	public static int deviceScreenOrientation = DEVICE_ORIENTATION_UNKNOWN;
+
 	public static enum VirtualKeyboardCommand { VK_CMD_NONE, VK_CMD_SHOW, VK_CMD_HIDE };
 	public static VirtualKeyboardCommand lastVirtualKeyboardCommand = VirtualKeyboardCommand.VK_CMD_NONE;
 	public static final int lastVirtualKeyboardCommandDelay = 200;
@@ -2940,7 +2954,7 @@ public class GameActivity extends $${gameActivitySuperClass}$$ implements Surfac
 		nativeSetGlobalActivity(UseExternalFilesDir, PublicLogFiles, getFilesDir().getPath(), getExternalFilesDir(null).getPath(), PackageDataInsideApkValue==1, APKPath);
 
 		// tell the engine if this is a portrait app
-		nativeSetWindowInfo(getResources().getConfiguration().orientation, DepthBufferPreference, PropagateAlpha);
+		nativeSetWindowInfo(getResources().getConfiguration().orientation == Configuration.ORIENTATION_PORTRAIT, DepthBufferPreference, PropagateAlpha);
 
 		// get the full language code, like en-US
 		// note: this may need to be Locale.getDefault().getLanguage()
@@ -3383,7 +3397,7 @@ public class GameActivity extends $${gameActivitySuperClass}$$ implements Surfac
 		}
 
 		// invalidate window cache
-		nativeSetWindowInfo(getResources().getConfiguration().orientation, DepthBufferPreference, PropagateAlpha);
+		nativeSetWindowInfo(getResources().getConfiguration().orientation == Configuration.ORIENTATION_PORTRAIT, DepthBufferPreference, PropagateAlpha);
 
 		// only do this on KitKat and above
 		if (ShouldHideUI)
@@ -3878,7 +3892,7 @@ public class GameActivity extends $${gameActivitySuperClass}$$ implements Surfac
 		super.onConfigurationChanged(newConfig);
 		//$${gameActivityonConfigurationChangedAdditions}$$
 		// forward the orientation
-		nativeOnConfigurationChanged(newConfig.orientation);
+		nativeOnOrientationChanged(newConfig.orientation);
 	}
 
 	public void surfaceChanged(SurfaceHolder holder, int format, int width, int height)
@@ -6790,9 +6804,39 @@ public class GameActivity extends $${gameActivitySuperClass}$$ implements Surfac
 		return getWindowManager().getDefaultDisplay().getRotation();
 	}
 
-	public int AndroidThunkJava_GetConfigurationOrientation()
+	public int AndroidThunkJava_GetDeviceOrientation()
 	{
-		return getResources().getConfiguration().orientation;
+		int currentScreenOrientation = DEVICE_ORIENTATION_UNKNOWN;
+		int rotation = getWindowManager().getDefaultDisplay().getRotation();
+		if ( getDeviceDefaultOrientation() == Configuration.ORIENTATION_LANDSCAPE )
+		{
+			switch (rotation)
+			{
+				case android.view.Surface.ROTATION_0:	currentScreenOrientation = DEVICE_ORIENTATION_LANDSCAPE_LEFT;            break;
+				case android.view.Surface.ROTATION_90:	currentScreenOrientation = DEVICE_ORIENTATION_PORTRAIT_UPSIDEDOWN;       break;
+				case android.view.Surface.ROTATION_180:	currentScreenOrientation = DEVICE_ORIENTATION_LANDSCAPE_RIGHT;			break;
+				case android.view.Surface.ROTATION_270:	currentScreenOrientation = DEVICE_ORIENTATION_PORTRAIT;					break;
+			}
+		}
+		else
+		{
+			switch (rotation)
+			{
+				case android.view.Surface.ROTATION_0:	currentScreenOrientation = DEVICE_ORIENTATION_PORTRAIT;					break;
+				case android.view.Surface.ROTATION_90:	currentScreenOrientation = DEVICE_ORIENTATION_LANDSCAPE_LEFT;			break;
+				case android.view.Surface.ROTATION_180:	currentScreenOrientation = DEVICE_ORIENTATION_PORTRAIT_UPSIDEDOWN;		break;
+				case android.view.Surface.ROTATION_270:	currentScreenOrientation = DEVICE_ORIENTATION_LANDSCAPE_RIGHT;			break;
+			}
+		}
+
+		if (deviceScreenOrientation != currentScreenOrientation)
+		{
+			deviceScreenOrientation = currentScreenOrientation;
+			Log.debug("==============> onSensorChanged: orientation : " + deviceScreenOrientation);
+			nativeOnOrientationChanged(deviceScreenOrientation);
+		}
+
+		return deviceScreenOrientation;
 	}
 
 	@SuppressWarnings("deprecation")
@@ -6925,7 +6969,7 @@ public class GameActivity extends $${gameActivitySuperClass}$$ implements Surfac
 	public native void nativeSetAndroidStartupState(boolean bDebuggerAttached);
 	public native void nativeSetGlobalActivity(boolean bUseExternalFilesDir, boolean bPublicLogFiles, String internalFilePath, String externalFilePath, boolean bOBBInAPK, String APKPath);
 	public native void nativeSetObbFilePaths(String OBBMainFilePath, String OBBPatchFilePath, String OBBOverflow1FilePath, String OBBOverflow2FilePath);
-	public native void nativeSetWindowInfo(int orientation, int DepthBufferPreference, int PropagateAlpha);
+	public native void nativeSetWindowInfo(boolean bIsPortrait, int DepthBufferPreference, int PropagateAlpha);
 	public native void nativeSetObbInfo(String ProjectName, String PackageName, int Version, int PatchVersion, String AppType);
 	public native void nativeSetAndroidVersionInformation( String AndroidVersion, int TargetSDK, String PhoneMake, String PhoneModel, String PhoneBuildNumber, String OSLanguage );
 
@@ -6950,7 +6994,7 @@ public class GameActivity extends $${gameActivitySuperClass}$$ implements Surfac
 	public native void nativeVirtualKeyboardShown(int left, int top, int right, int bottom);
 	public native void nativeVirtualKeyboardVisible(boolean bShown);
 
-	public native void nativeOnConfigurationChanged(int orientation);
+	public native void nativeOnOrientationChanged(int orientation);
 		
 	public native void nativeOnInitialDownloadStarted();
 	public native void nativeOnInitialDownloadCompleted();
diff --git a/Engine/Source/Runtime/ApplicationCore/Private/Android/AndroidApplication.cpp b/Engine/Source/Runtime/ApplicationCore/Private/Android/AndroidApplication.cpp
--- a/Engine/Source/Runtime/ApplicationCore/Private/Android/AndroidApplication.cpp
+++ b/Engine/Source/Runtime/ApplicationCore/Private/Android/AndroidApplication.cpp
@@ -107,12 +107,6 @@ void FAndroidApplication::HandleDeviceOrientation()
 	{
 		FCoreDelegates::ApplicationReceivedScreenOrientationChangedNotificationDelegate.Broadcast((int32)DeviceOrientation);
 	}
-
-	if (FAndroidWindow::SafezoneUpdated())
-	{
-		//we also want to fire off the safe frame event
-		FCoreDelegates::OnSafeFrameChangedEvent.Broadcast();
-	}
 }
 
 FPlatformRect FAndroidApplication::GetWorkArea( const FPlatformRect& CurrentWindow ) const
diff --git a/Engine/Source/Runtime/ApplicationCore/Private/Android/AndroidWindow.cpp b/Engine/Source/Runtime/ApplicationCore/Private/Android/AndroidWindow.cpp
--- a/Engine/Source/Runtime/ApplicationCore/Private/Android/AndroidWindow.cpp
+++ b/Engine/Source/Runtime/ApplicationCore/Private/Android/AndroidWindow.cpp
@@ -138,17 +138,14 @@ void FAndroidWindow::SetOSWindowHandle(void* InWindow)
 
 //This function is declared in the Java-defined class, GameActivity.java: "public native void nativeSetObbInfo(String PackageName, int Version, int PatchVersion);"
 static bool GAndroidIsPortrait = false;
-static int GWindowOrientation = -1;
-static std::atomic<bool> GAndroidSafezoneRequiresUpdate = false;
+static EDeviceScreenOrientation GDeviceScreenOrientation = EDeviceScreenOrientation::Unknown;
 static int GAndroidDepthBufferPreference = 0;
 static FVector4 GAndroidPortraitSafezone = FVector4(-1.0f, -1.0f, -1.0f, -1.0f);
 static FVector4 GAndroidLandscapeSafezone = FVector4(-1.0f, -1.0f, -1.0f, -1.0f);
 #if USE_ANDROID_JNI
-JNI_METHOD void Java_com_epicgames_unreal_GameActivity_nativeSetWindowInfo(JNIEnv* jenv, jobject thiz, jint orientation, jint DepthBufferPreference, jint PropagateAlpha)
+JNI_METHOD void Java_com_epicgames_unreal_GameActivity_nativeSetWindowInfo(JNIEnv* jenv, jobject thiz, jboolean bIsPortrait, jint DepthBufferPreference, jint PropagateAlpha)
 {
 	ClearCachedWindowRects();
-	GWindowOrientation = orientation;
-	bool bIsPortrait = GWindowOrientation == EAndroidConfigurationOrientation::ORIENTATION_PORTRAIT;
 	GAndroidIsPortrait = bIsPortrait == JNI_TRUE;
 	GAndroidDepthBufferPreference = DepthBufferPreference;
 	GAndroidPropagateAlpha = PropagateAlpha;
@@ -179,8 +176,9 @@ JNI_METHOD void Java_com_epicgames_unreal_GameActivity_nativeSetSafezoneInfo(JNI
 		GAndroidLandscapeSafezone.Z = right;
 		GAndroidLandscapeSafezone.W = bottom;
 	}
-
-	GAndroidSafezoneRequiresUpdate = true;
+#if USE_ANDROID_EVENTS
+	FAppEventManager::GetInstance()->EnqueueAppEvent(APP_EVENT_STATE_SAFE_ZONE_UPDATED);
+#endif
 	UE_LOG(LogAndroid, Log, TEXT("nativeSetSafezoneInfo bIsPortrait=%d, left=%f, top=%f, right=%f, bottom=%f"), bIsPortrait ? 1 : 0, left, top, right, bottom);
 }
 #endif
@@ -199,12 +197,6 @@ FVector4 FAndroidWindow::GetSafezone(bool bPortrait)
 	return bPortrait ? GAndroidPortraitSafezone : GAndroidLandscapeSafezone;
 }
 
-bool FAndroidWindow::SafezoneUpdated()
-{
-	bool bRequiresUpdate = true;
-	return GAndroidSafezoneRequiresUpdate.compare_exchange_weak(bRequiresUpdate, false);
-}
-
 int32 FAndroidWindow::GetDepthBufferPreference()
 {
 	return GAndroidDepthBufferPreference;
@@ -570,13 +562,13 @@ void FAndroidWindow::CalculateSurfaceSize(int32_t& SurfaceWidth, int32_t& Surfac
 #endif
 }
 
-bool FAndroidWindow::OnWindowOrientationChanged(int Orientation)
+bool FAndroidWindow::OnWindowOrientationChanged(EDeviceScreenOrientation DeviceScreenOrientation)
 {
-	if (GWindowOrientation != Orientation)
+	if (GDeviceScreenOrientation != DeviceScreenOrientation)
 	{
-		GWindowOrientation = Orientation;
-		bool bIsPortrait = GWindowOrientation == EAndroidConfigurationOrientation::ORIENTATION_PORTRAIT;
-		UE_LOG(LogAndroid, Log, TEXT("Window orientation changed: %s"), bIsPortrait ? TEXT("Portrait") : TEXT("Landscape"));
+		GDeviceScreenOrientation = DeviceScreenOrientation;
+		bool bIsPortrait = GDeviceScreenOrientation == EDeviceScreenOrientation::Portrait || GDeviceScreenOrientation == EDeviceScreenOrientation::PortraitUpsideDown;
+		UE_LOG(LogAndroid, Log, TEXT("Window orientation changed: %s, GDeviceScreenOrientation=%d"), bIsPortrait ? TEXT("Portrait") : TEXT("Landscape"), GDeviceScreenOrientation);
 		GAndroidIsPortrait = bIsPortrait;
 		return true;
 	}
diff --git a/Engine/Source/Runtime/ApplicationCore/Public/Android/AndroidWindow.h b/Engine/Source/Runtime/ApplicationCore/Public/Android/AndroidWindow.h
--- a/Engine/Source/Runtime/ApplicationCore/Public/Android/AndroidWindow.h
+++ b/Engine/Source/Runtime/ApplicationCore/Public/Android/AndroidWindow.h
@@ -4,6 +4,7 @@
 
 #include "GenericPlatform/GenericWindow.h"
 #include "GenericPlatform/GenericApplication.h"
+#include "GenericPlatform/GenericPlatformMisc.h"
 #include <android/native_window.h> 
 #if USE_ANDROID_JNI
 #include <android/native_window_jni.h>
@@ -43,7 +44,7 @@ public:
 	// When bUseEventThreadWindow == false this uses dimensions cached when the game thread processes android events.
 	// When bUseEventThreadWindow == true this uses dimensions directly from the android event thread, unless called from event thread this requires acquiring GAndroidWindowLock to use.
 	static APPLICATIONCORE_API void CalculateSurfaceSize(int32_t& SurfaceWidth, int32_t& SurfaceHeight, bool bUseEventThreadWindow = false);
-	static APPLICATIONCORE_API bool OnWindowOrientationChanged(int Orientation);
+	static APPLICATIONCORE_API bool OnWindowOrientationChanged(EDeviceScreenOrientation DeviceScreenOrientation);
 
 	static APPLICATIONCORE_API int32 GetDepthBufferPreference();
 	
@@ -61,7 +62,6 @@ public:
 
 	static APPLICATIONCORE_API bool IsPortraitOrientation();
 	static APPLICATIONCORE_API FVector4 GetSafezone(bool bPortrait);
-	static APPLICATIONCORE_API bool SafezoneUpdated();
 
 	// called by the Android event thread to initially set the current window dimensions.
 	static APPLICATIONCORE_API void SetWindowDimensions_EventThread(ANativeWindow* DimensionWindow);
diff --git a/Engine/Source/Runtime/Core/Private/Android/AndroidPlatformMisc.cpp b/Engine/Source/Runtime/Core/Private/Android/AndroidPlatformMisc.cpp
--- a/Engine/Source/Runtime/Core/Private/Android/AndroidPlatformMisc.cpp
+++ b/Engine/Source/Runtime/Core/Private/Android/AndroidPlatformMisc.cpp
@@ -618,44 +618,21 @@ void FAndroidMisc::UpdateDeviceOrientation()
 	JNIEnv* JEnv = AndroidJavaEnv::GetJavaEnv();
 	if (JEnv)
 	{
-		static jmethodID getRotationMethod = 0;
 		static jmethodID getOrientationMethod = 0;
 
-		if (getRotationMethod == 0 || getOrientationMethod == 0)
+		if (getOrientationMethod == 0)
 		{
 			jclass MainClass = AndroidJavaEnv::FindJavaClassGlobalRef("com/epicgames/unreal/GameActivity");
 			if (MainClass != nullptr)
 			{
-				getRotationMethod = JEnv->GetMethodID(MainClass, "AndroidThunkJava_GetDeviceRotation", "()I");
-				getOrientationMethod = JEnv->GetMethodID(MainClass, "AndroidThunkJava_GetConfigurationOrientation", "()I");
+				getOrientationMethod = JEnv->GetMethodID(MainClass, "AndroidThunkJava_GetDeviceOrientation", "()I");
 				JEnv->DeleteGlobalRef(MainClass);
 			}
 		}
 
-		if (getRotationMethod != 0 && getOrientationMethod != 0)
+		if (getOrientationMethod != 0)
 		{
-			const int Rotation = JEnv->CallIntMethod(AndroidJavaEnv::GetGameActivityThis(), getRotationMethod);
-			const int Orientation = JEnv->CallIntMethod(AndroidJavaEnv::GetGameActivityThis(), getOrientationMethod);
-			if (Orientation == EAndroidConfigurationOrientation::ORIENTATION_PORTRAIT)
-			{
-				switch (Rotation)
-				{
-				case EAndroidSurfaceRotation::ROTATION_0:	DeviceOrientation = EDeviceScreenOrientation::Portrait;             break;
-				case EAndroidSurfaceRotation::ROTATION_90:	DeviceOrientation = EDeviceScreenOrientation::LandscapeLeft;        break;
-				case EAndroidSurfaceRotation::ROTATION_180:	DeviceOrientation = EDeviceScreenOrientation::PortraitUpsideDown;   break;
-				case EAndroidSurfaceRotation::ROTATION_270:	DeviceOrientation = EDeviceScreenOrientation::LandscapeRight;       break;
-				}
-			}
-			else if (Orientation == EAndroidConfigurationOrientation::ORIENTATION_LANDSCAPE)
-			{
-				switch (Rotation)
-				{
-				case EAndroidSurfaceRotation::ROTATION_0:	DeviceOrientation = EDeviceScreenOrientation::LandscapeLeft;        break;
-				case EAndroidSurfaceRotation::ROTATION_90:	DeviceOrientation = EDeviceScreenOrientation::PortraitUpsideDown;				break;
-				case EAndroidSurfaceRotation::ROTATION_180:	DeviceOrientation = EDeviceScreenOrientation::LandscapeRight;		break;
-				case EAndroidSurfaceRotation::ROTATION_270:	DeviceOrientation = EDeviceScreenOrientation::Portrait;   break;
-				}
-			}
+			DeviceOrientation = (EDeviceScreenOrientation)JEnv->CallIntMethod(AndroidJavaEnv::GetGameActivityThis(), getOrientationMethod);
 		}
 	}
 #endif
diff --git a/Engine/Source/Runtime/Core/Public/Android/AndroidPlatformMisc.h b/Engine/Source/Runtime/Core/Public/Android/AndroidPlatformMisc.h
--- a/Engine/Source/Runtime/Core/Public/Android/AndroidPlatformMisc.h
+++ b/Engine/Source/Runtime/Core/Public/Android/AndroidPlatformMisc.h
@@ -23,28 +23,6 @@ class TFunction;
 
 enum class ECrashContextType;
 
-//The android configuration orientations defined in C++ map to the values in https://developer.android.com/reference/android/content/res/Configuration#ORIENTATION_PORTRAIT
-namespace EAndroidConfigurationOrientation
-{
-	enum Type
-	{
-		ORIENTATION_PORTRAIT = 0x00000001,
-		ORIENTATION_LANDSCAPE = 0x00000002,
-	};
-};
-
-//The android surface rotations defined in C++ map to the values in https://developer.android.com/reference/android/view/Surface#ROTATION_0
-namespace EAndroidSurfaceRotation
-{
-	enum Type
-	{
-		ROTATION_0,
-		ROTATION_90,
-		ROTATION_180,
-		ROTATION_270,
-	};
-};
-
 /**
  * Android implementation of the misc OS functions
  */
diff --git a/Engine/Source/Runtime/Launch/Private/Android/AndroidEventManager.cpp b/Engine/Source/Runtime/Launch/Private/Android/AndroidEventManager.cpp
--- a/Engine/Source/Runtime/Launch/Private/Android/AndroidEventManager.cpp
+++ b/Engine/Source/Runtime/Launch/Private/Android/AndroidEventManager.cpp
@@ -87,7 +87,7 @@ void FAppEventManager::Tick()
 		case APP_EVENT_STATE_WINDOW_CHANGED:
 			// React on device orientation/windowSize changes only when application has window
 			// In case window was created this tick it should already has correct size
-			// see 'Java_com_epicgames_unreal_GameActivity_nativeOnConfigurationChanged' for event thread/game thread mismatches.
+			// see 'Java_com_epicgames_unreal_GameActivity_nativeOnOrientationChanged' for event thread/game thread mismatches.
 			ExecWindowResized();
 		break;
 		case APP_EVENT_STATE_SAVE_STATE:
@@ -138,6 +138,9 @@ void FAppEventManager::Tick()
 			bRunning = false;
 			FPlatformMisc::LowLevelOutputDebugStringf(TEXT("Execution will be paused..."));
 			break;
+		case APP_EVENT_STATE_SAFE_ZONE_UPDATED:
+			FCoreDelegates::OnSafeFrameChangedEvent.Broadcast();
+			break;
 		default:
 			UE_LOG(LogAndroidEvents, Display, TEXT("Application Event : %u  not handled. "), Event.State);
 		}
diff --git a/Engine/Source/Runtime/Launch/Private/Android/LaunchAndroid.cpp b/Engine/Source/Runtime/Launch/Private/Android/LaunchAndroid.cpp
--- a/Engine/Source/Runtime/Launch/Private/Android/LaunchAndroid.cpp
+++ b/Engine/Source/Runtime/Launch/Private/Android/LaunchAndroid.cpp
@@ -1663,15 +1663,16 @@ JNI_METHOD jboolean Java_com_epicgames_unreal_GameActivity_nativeSupportsNEON(JN
 	return JNI_FALSE;
 }
 
-//This function is declared in the Java-defined class, GameActivity.java: "public native void nativeOnConfigurationChanged(boolean bPortrait);
-JNI_METHOD void Java_com_epicgames_unreal_GameActivity_nativeOnConfigurationChanged(JNIEnv* jenv, jobject thiz, jint orientation)
+//This function is declared in the Java-defined class, GameActivity.java: "public native void nativeOnOrientationChanged(int orientation);
+JNI_METHOD void Java_com_epicgames_unreal_GameActivity_nativeOnOrientationChanged(JNIEnv* jenv, jobject thiz, jint orientation)
 {
 	// enqueue a window changed event if orientation changed, 
 	// note that the HW window handle does not necessarily change.
-	if (FAndroidWindow::OnWindowOrientationChanged(orientation))
+	if (FAndroidWindow::OnWindowOrientationChanged((EDeviceScreenOrientation)orientation))
 	{	
 		// Enqueue an event to trigger gamethread to update the orientation:
 		FAppEventManager::GetInstance()->EnqueueAppEvent(APP_EVENT_STATE_WINDOW_CHANGED);
+		FAppEventManager::GetInstance()->EnqueueAppEvent(APP_EVENT_STATE_SAFE_ZONE_UPDATED);
 
 		if (EventHandlerEvent)
 		{
diff --git a/Engine/Source/Runtime/Launch/Public/Android/AndroidEventManager.h b/Engine/Source/Runtime/Launch/Public/Android/AndroidEventManager.h
--- a/Engine/Source/Runtime/Launch/Public/Android/AndroidEventManager.h
+++ b/Engine/Source/Runtime/Launch/Public/Android/AndroidEventManager.h
@@ -28,6 +28,7 @@ enum EAppEventState
 	APP_EVENT_STATE_SAVE_STATE,
 	APP_EVENT_STATE_APP_SUSPENDED,
 	APP_EVENT_STATE_APP_ACTIVATED,
+	APP_EVENT_STATE_SAFE_ZONE_UPDATED,
 	APP_EVENT_RUN_CALLBACK,
 	APP_EVENT_STATE_INVALID = -1,
 };
