Fix Android PSO program service initialization failure when application is backgrounded during engine init.
 - service timeout error was thrown since connection does not complete when app is backgrounded.

diff --git a/Engine/Build/Android/Java/src/com/epicgames/unreal/psoservices/PSOProgramServiceAccessor.java b/Engine/Build/Android/Java/src/com/epicgames/unreal/psoservices/PSOProgramServiceAccessor.java
--- a/Engine/Build/Android/Java/src/com/epicgames/unreal/psoservices/PSOProgramServiceAccessor.java
+++ b/Engine/Build/Android/Java/src/com/epicgames/unreal/psoservices/PSOProgramServiceAccessor.java
@@ -80,7 +80,8 @@ public class PSOProgramServiceAccessor
 	private static PSOProgramServiceAccessor _PSOProgramServiceAccessor = null;
 
 	OGLServiceInstance[] ServiceInstances;
-	private GameActivity mContext;
+	private GameActivity mActivity;
+	private Context mContext;
 
 	static class JobResponse
 	{
@@ -175,16 +176,17 @@ public class PSOProgramServiceAccessor
 	}
 	
 	private static final Logger Log = GameActivity.Log;
-	public void Init(GameActivity InContext)
+	public void Init(GameActivity InActivity)
 	{
-		mContext = InContext;
+		mActivity = InActivity;
+		mContext = mActivity.getApplicationContext();
 		ClearServiceLogDirs();
 		_PSOProgramServiceAccessor = this;
 	}
 	
 	private void ClearServiceLogDirs()
 	{
-		String OGLLogDir = mContext.getFilesDir()+PSOProgramService.LogDir;
+		String OGLLogDir = mActivity.getFilesDir()+PSOProgramService.LogDir;
 		File deleteDir = new File(OGLLogDir);
 		boolean bDeleteDirExists = deleteDir.exists();
 		if( bDeleteDirExists )
@@ -675,7 +677,7 @@ public class PSOProgramServiceAccessor
 		// Opens the log file(s) for the service, scans each entry and prints any new entries.  
 		synchronized void ReadBackServiceLog()
 		{
-			String OGLLogDirPath = mContext.getFilesDir()+PSOProgramService.LogDir;
+			String OGLLogDirPath = mActivity.getFilesDir()+PSOProgramService.LogDir;
 			File OGLLogDir = new File(OGLLogDirPath);
 			File[] LogFiles = OGLLogDir.listFiles();
 			if( LogFiles == null)
