Add null clipboard item check

#android

diff --git a/Engine/Build/Android/Java/src/com/epicgames/unreal/GameActivity.java.template b/Engine/Build/Android/Java/src/com/epicgames/unreal/GameActivity.java.template
--- a/Engine/Build/Android/Java/src/com/epicgames/unreal/GameActivity.java.template
+++ b/Engine/Build/Android/Java/src/com/epicgames/unreal/GameActivity.java.template
@@ -6791,10 +6791,17 @@ public class GameActivity extends $${gameActivitySuperClass}$$ implements Surfac
 
 	public String AndroidThunkJava_ClipboardPaste() 
 	{
-		if(clipboardManager.hasPrimaryClip()) 
+		if (clipboardManager.hasPrimaryClip()) 
 		{
 			ClipData.Item item = clipboardManager.getPrimaryClip().getItemAt(0);
-			return item.getText().toString();
+			if (item != null)
+			{
+				CharSequence rawTextSequence = item.getText();
+				if (rawTextSequence != null && rawTextSequence.length() > 0)
+				{
+					return rawTextSequence.toString();
+				}
+			}
 		} 
 		return "";
 	}
