Android: Add 'Product name' to android misc, gives access to the device's marketing/product name - if available.

diff --git a/Engine/Build/Android/Java/src/com/epicgames/unreal/GameActivity.java.template b/Engine/Build/Android/Java/src/com/epicgames/unreal/GameActivity.java.template
--- a/Engine/Build/Android/Java/src/com/epicgames/unreal/GameActivity.java.template
+++ b/Engine/Build/Android/Java/src/com/epicgames/unreal/GameActivity.java.template
@@ -638,6 +638,16 @@ public class GameActivity extends $${gameActivitySuperClass}$$ implements Surfac
 		return (resourceId < 1) ? DefaultString : getString(resourceId);
 	}
 
+	private String getProductName()
+	{
+		String ProductName = Settings.Global.getString(getContentResolver(), "default_device_name");
+		if(ProductName == null)
+		{
+			ProductName="";
+		}
+		return ProductName;
+	}
+
     public static List<String> getInstalledPackages(final Context context)
     {
         List<String> result = new ArrayList<>();
@@ -2963,14 +2973,16 @@ public class GameActivity extends $${gameActivitySuperClass}$$ implements Surfac
 		boolean bDebuggerAttached = android.os.Debug.isDebuggerConnected();
 		nativeSetAndroidStartupState(bDebuggerAttached);
 
+		String productName = getProductName();
 		Log.debug( "Android version is " + android.os.Build.VERSION.RELEASE );
 		Log.debug( "Android manufacturer is " + android.os.Build.MANUFACTURER );
 		Log.debug( "Android model is " + android.os.Build.MODEL );
 		Log.debug( "Android build number is " + android.os.Build.DISPLAY );
 		Log.debug( "OS language is set to " + Language );
+		Log.debug( "Product name is " + (productName.isEmpty() ? "[not set]" : productName) );
 		Log.debug( "Debugger attached is " + bDebuggerAttached );
 
-		nativeSetAndroidVersionInformation( android.os.Build.VERSION.RELEASE, targetSdkVersion, android.os.Build.MANUFACTURER, android.os.Build.MODEL, android.os.Build.DISPLAY, Language );
+		nativeSetAndroidVersionInformation( android.os.Build.VERSION.RELEASE, targetSdkVersion, android.os.Build.MANUFACTURER, android.os.Build.MODEL, android.os.Build.DISPLAY, Language, productName );
 		int PatchVersion = 0;
 		if (VersionCode > 0)
 		{
@@ -6978,7 +6990,7 @@ public class GameActivity extends $${gameActivitySuperClass}$$ implements Surfac
 	public native void nativeSetObbFilePaths(String OBBMainFilePath, String OBBPatchFilePath, String OBBOverflow1FilePath, String OBBOverflow2FilePath);
 	public native void nativeSetWindowInfo(boolean bIsPortrait, int DepthBufferPreference, int PropagateAlpha);
 	public native void nativeSetObbInfo(String ProjectName, String PackageName, int Version, int PatchVersion, String AppType);
-	public native void nativeSetAndroidVersionInformation( String AndroidVersion, int TargetSDK, String PhoneMake, String PhoneModel, String PhoneBuildNumber, String OSLanguage );
+	public native void nativeSetAndroidVersionInformation( String AndroidVersion, int TargetSDK, String PhoneMake, String PhoneModel, String PhoneBuildNumber, String OSLanguage, String ProductName );
 
 	public native void nativeSetSurfaceViewInfo(int width, int height);
 	public native void nativeSetSafezoneInfo(boolean bIsPortrait, float left, float top, float right, float bottom);
diff --git a/Engine/Source/Runtime/Core/Private/Android/AndroidPlatformMisc.cpp b/Engine/Source/Runtime/Core/Private/Android/AndroidPlatformMisc.cpp
--- a/Engine/Source/Runtime/Core/Private/Android/AndroidPlatformMisc.cpp
+++ b/Engine/Source/Runtime/Core/Private/Android/AndroidPlatformMisc.cpp
@@ -122,6 +122,7 @@ FString FAndroidMisc::DeviceMake; // make of the device we are running on eg. "s
 FString FAndroidMisc::DeviceModel; // model of the device we are running on eg "SAMSUNG-SGH-I437"
 FString FAndroidMisc::DeviceBuildNumber; // platform image build number of device "R16NW.G960NKSU1ARD6"
 FString FAndroidMisc::OSLanguage; // language code the device is set to eg "deu"
+FString FAndroidMisc::ProductName; // product/marketing name of the device, if available.
 
 // Build/API level we are running.
 int32 FAndroidMisc::AndroidBuildVersion = 0;
@@ -1696,7 +1697,7 @@ bool FAndroidMisc::FileExistsInPlatformPackage(const FString& RelativePath)
 	return false;
 }
 
-void FAndroidMisc::SetVersionInfo( FString InAndroidVersion, int32 InTargetSDKVersion, FString InDeviceMake, FString InDeviceModel, FString InDeviceBuildNumber, FString InOSLanguage )
+void FAndroidMisc::SetVersionInfo( FString InAndroidVersion, int32 InTargetSDKVersion, FString InDeviceMake, FString InDeviceModel, FString InDeviceBuildNumber, FString InOSLanguage, FString InProductName)
 {
 	AndroidVersion = InAndroidVersion;
 	AndroidMajorVersion = FCString::Atoi(*InAndroidVersion);
@@ -1705,8 +1706,8 @@ void FAndroidMisc::SetVersionInfo( FString InAndroidVersion, int32 InTargetSDKVe
 	DeviceModel = InDeviceModel;
 	DeviceBuildNumber = InDeviceBuildNumber;
 	OSLanguage = InOSLanguage;
-
-	UE_LOG(LogAndroid, Display, TEXT("Android Version Make Model BuildNumber Language: %s %s %s %s %s"), *AndroidVersion, *DeviceMake, *DeviceModel, *DeviceBuildNumber, *OSLanguage);
+	ProductName = InProductName;
+	UE_LOG(LogAndroid, Display, TEXT("Android Version: %s, Make: %s, Model: %s, BuildNumber: %s, Language: %s, Product name: %s"), *AndroidVersion, *DeviceMake, *DeviceModel, *DeviceBuildNumber, *OSLanguage, ProductName.IsEmpty() ? TEXT("[not set]") : *ProductName);
 }
 
 const FString FAndroidMisc::GetAndroidVersion()
@@ -1744,6 +1745,11 @@ const FString FAndroidMisc::GetOSLanguage()
 	return OSLanguage;
 }
 
+const FString FAndroidMisc::GetProductName()
+{
+	return ProductName;
+}
+
 const FString FAndroidMisc::GetProjectVersion() {
 	return FString::FromInt(GAndroidPackageVersion);
 }
diff --git a/Engine/Source/Runtime/Core/Public/Android/AndroidPlatformMisc.h b/Engine/Source/Runtime/Core/Public/Android/AndroidPlatformMisc.h
--- a/Engine/Source/Runtime/Core/Public/Android/AndroidPlatformMisc.h
+++ b/Engine/Source/Runtime/Core/Public/Android/AndroidPlatformMisc.h
@@ -167,13 +167,14 @@ public:
 	static CORE_API bool FileExistsInPlatformPackage(const FString& RelativePath);
 
 	// ANDROID ONLY:
-	static CORE_API void SetVersionInfo(FString AndroidVersion, int32 InTargetSDKVersion, FString DeviceMake, FString DeviceModel, FString DeviceBuildNumber, FString OSLanguage);
+	static CORE_API void SetVersionInfo(FString AndroidVersion, int32 InTargetSDKVersion, FString DeviceMake, FString DeviceModel, FString DeviceBuildNumber, FString OSLanguage, FString ProductName);
 	static CORE_API const FString GetAndroidVersion();
 	static CORE_API int32 GetAndroidMajorVersion();
 	static CORE_API int32 GetTargetSDKVersion();
 	static CORE_API const FString GetDeviceMake();
 	static CORE_API const FString GetDeviceModel();
 	static CORE_API const FString GetOSLanguage();
+	static CORE_API const FString GetProductName(); // returns the product name, if available. e.g. 'Galaxy Tab S8' or empty string.
 	static CORE_API const FString GetDeviceBuildNumber();
 	static CORE_API const FString GetProjectVersion();
 	static CORE_API FString GetDefaultLocale();
@@ -258,6 +259,7 @@ public:
 	static CORE_API FString DeviceModel; // model of the device we are running on eg "SAMSUNG-SGH-I437"
 	static CORE_API FString DeviceBuildNumber; // platform image build number of device "R16NW.G960NKSU1ARD6"
 	static CORE_API FString OSLanguage; // language code the device is set to
+	static CORE_API FString ProductName; // Product name, if available. e.g. 'Galaxy Tab S8' or empty string.
 
 	// Build version of Android, i.e. API level.
 	static CORE_API int32 AndroidBuildVersion;
diff --git a/Engine/Source/Runtime/Launch/Private/Android/LaunchAndroid.cpp b/Engine/Source/Runtime/Launch/Private/Android/LaunchAndroid.cpp
--- a/Engine/Source/Runtime/Launch/Private/Android/LaunchAndroid.cpp
+++ b/Engine/Source/Runtime/Launch/Private/Android/LaunchAndroid.cpp
@@ -1712,15 +1712,16 @@ JNI_METHOD void Java_com_epicgames_unreal_GameActivity_nativeInitHMDs(JNIEnv* je
 	GHMDsInitialized = true;
 }
 
-JNI_METHOD void Java_com_epicgames_unreal_GameActivity_nativeSetAndroidVersionInformation(JNIEnv* jenv, jobject thiz, jstring androidVersion, jint targetSDKversion, jstring phoneMake, jstring phoneModel, jstring phoneBuildNumber, jstring osLanguage)
+JNI_METHOD void Java_com_epicgames_unreal_GameActivity_nativeSetAndroidVersionInformation(JNIEnv* jenv, jobject thiz, jstring androidVersion, jint targetSDKversion, jstring phoneMake, jstring phoneModel, jstring phoneBuildNumber, jstring osLanguage, jstring productName)
 {
 	auto UEAndroidVersion = FJavaHelper::FStringFromParam(jenv, androidVersion);
 	auto UEPhoneMake = FJavaHelper::FStringFromParam(jenv, phoneMake);
 	auto UEPhoneModel = FJavaHelper::FStringFromParam(jenv, phoneModel);
 	auto UEPhoneBuildNumber = FJavaHelper::FStringFromParam(jenv, phoneBuildNumber);
 	auto UEOSLanguage = FJavaHelper::FStringFromParam(jenv, osLanguage);
+	auto UEProductName = FJavaHelper::FStringFromParam(jenv, productName);
 
-	FAndroidMisc::SetVersionInfo(UEAndroidVersion, targetSDKversion, UEPhoneMake, UEPhoneModel, UEPhoneBuildNumber, UEOSLanguage);
+	FAndroidMisc::SetVersionInfo(UEAndroidVersion, targetSDKversion, UEPhoneMake, UEPhoneModel, UEPhoneBuildNumber, UEOSLanguage, UEProductName);
 }
 
 //This function is declared in the Java-defined class, GameActivity.java: "public native void nativeOnInitialDownloadStarted();
