Reduce PSO compiler service priority.
 - reduces likelihood of PSO tasks starving application of cpu cycles.

diff --git a/Engine/Build/Android/Java/src/com/epicgames/unreal/GameActivity.java.template b/Engine/Build/Android/Java/src/com/epicgames/unreal/GameActivity.java.template
--- a/Engine/Build/Android/Java/src/com/epicgames/unreal/GameActivity.java.template
+++ b/Engine/Build/Android/Java/src/com/epicgames/unreal/GameActivity.java.template
@@ -2327,6 +2327,11 @@ public class GameActivity extends $${gameActivitySuperClass}$$ implements Surfac
 		Log.debug("LittleCoreMask: 0x" + Integer.toHexString(littleCoreMask));
 		nativeSetAffinityInfo(bUseAffinity, bigCoreMask, littleCoreMask);
 
+		if (variables.containsKey("PSOServicePri"))
+		{
+			mPSOProgramServiceAccessor.SetServicePriority(true);
+		}
+
 		// send the variables to the engine
 		String[] keyValues = new String[variables.size() * 2];
 		int keyValueIndex = 0;
diff --git a/Engine/Build/Android/Java/src/com/epicgames/unreal/psoservices/PSOProgramServiceAccessor.java b/Engine/Build/Android/Java/src/com/epicgames/unreal/psoservices/PSOProgramServiceAccessor.java
--- a/Engine/Build/Android/Java/src/com/epicgames/unreal/psoservices/PSOProgramServiceAccessor.java
+++ b/Engine/Build/Android/Java/src/com/epicgames/unreal/psoservices/PSOProgramServiceAccessor.java
@@ -167,6 +167,9 @@ public class PSOProgramServiceAccessor
 
 	private static boolean IsNullOrEmpty(String string) { return string == null || string.isEmpty(); }
 
+	private boolean bServicePriHigh = false;
+	public void SetServicePriority(boolean bHigh) { bServicePriHigh = bHigh; }
+
 	public PSOProgramServiceAccessor()
 	{
 		//StrictMode.enableDefaults();
@@ -976,7 +979,13 @@ public class PSOProgramServiceAccessor
 			{
 				Intent intent = new Intent(mContext, ServiceClass);
 
-				mShouldUnbind = mContext.bindService(intent, mConnection, Context.BIND_AUTO_CREATE | Context.BIND_IMPORTANT);
+				int BindServiceFlags = Context.BIND_AUTO_CREATE;
+				if(	bServicePriHigh )
+				{
+					BindServiceFlags = BindServiceFlags | Context.BIND_IMPORTANT;
+				}
+
+				mShouldUnbind = mContext.bindService(intent, mConnection, BindServiceFlags);
 
 				Log.verbose("doBindService " + Name() + " needs unbind " + mShouldUnbind + " bound: " + mBound);
 			}
