Add protection to InputDevice check for gamepads

#android

diff --git a/Engine/Build/Android/Java/src/com/epicgames/unreal/GameActivity.java.template b/Engine/Build/Android/Java/src/com/epicgames/unreal/GameActivity.java.template
--- a/Engine/Build/Android/Java/src/com/epicgames/unreal/GameActivity.java.template
+++ b/Engine/Build/Android/Java/src/com/epicgames/unreal/GameActivity.java.template
@@ -5626,6 +5626,12 @@ public class GameActivity extends $${gameActivitySuperClass}$$ implements Surfac
 		for (int deviceIndex=0; deviceIndex < deviceIds.length; deviceIndex++)
 		{
 			InputDevice inputDevice = InputDevice.getDevice(deviceIds[deviceIndex]);
+			// possible device id is invalid by time we try to use it above
+			if (inputDevice == null)
+			{
+				continue;
+			}
+
 			// is it a joystick, dpad, or gamepad?
 			if (((inputDevice.getSources() & InputDevice.SOURCE_GAMEPAD) == InputDevice.SOURCE_GAMEPAD)
                 || ((inputDevice.getSources() & InputDevice.SOURCE_JOYSTICK) == InputDevice.SOURCE_JOYSTICK))
