Android Single Instance Service missing files for packaging and shortname for build.cs
#android

diff --git a/Engine/Build/Android/Java/src/com/epicgames/unreal/GameActivity.java.template b/Engine/Build/Android/Java/src/com/epicgames/unreal/GameActivity.java.template
--- a/Engine/Build/Android/Java/src/com/epicgames/unreal/GameActivity.java.template
+++ b/Engine/Build/Android/Java/src/com/epicgames/unreal/GameActivity.java.template
@@ -342,9 +342,11 @@ public class GameActivity extends $${gameActivitySuperClass}$$ implements Surfac
 	public static final int checkLastVirtualKeyboardCommandDelay = 500;
 	public static final boolean bSamsungDevice = android.os.Build.MANUFACTURER.equals("samsung");
 	
-	static GameActivity _activity = null;
-	static Bundle _bundle;
-	static Bundle _extrasBundle;
+	protected boolean isStandalone = false;
+	protected static GameActivity _gameActivityInstance = null;
+	protected SimpleContextWrapper activityContext = null;
+	protected static Bundle _bundle;
+	protected static Bundle _extrasBundle;
 
 	private HandlerThread memoryHandlerThread;
 	private Handler memoryHandler;
@@ -416,8 +418,8 @@ public class GameActivity extends $${gameActivitySuperClass}$$ implements Surfac
 	ConsoleCmdReceiver consoleCmdReceiver;
 
 	// default the PackageDataInsideApk to an invalid value to make sure we don't get it too early
-	private static int PackageDataInsideApkValue = -1;
-	private static int HasOBBFiles = -1;
+	protected static int PackageDataInsideApkValue = -1;
+	protected static int HasOBBFiles = -1;
 	
 	// depthbuffer preference from manifest
 	int DepthBufferPreference = 0;
@@ -426,7 +428,7 @@ public class GameActivity extends $${gameActivitySuperClass}$$ implements Surfac
 	int PropagateAlpha = 0;
 
 	/** AssetManger reference - populated on start up and used when the OBB is packed into the APK */
-	private AssetManager			AssetManagerReference;
+	protected AssetManager			AssetManagerReference;
 	
 	private GoogleApiClient googleClient;
 
@@ -442,7 +444,7 @@ public class GameActivity extends $${gameActivitySuperClass}$$ implements Surfac
 	/** Unique ID to identify Google Play Services error dialog */
 	private static final int PLAY_SERVICES_DIALOG_ID = 1;
 	
-	private static String appPackageName = "";
+	protected static String appPackageName = "";
 
 	/** Check to see if we have all the files */
 	private boolean HasAllFiles = false;
@@ -485,7 +487,7 @@ public class GameActivity extends $${gameActivitySuperClass}$$ implements Surfac
 	private boolean ShouldHideUI = false;
 
 	/** Whether this application is for distribution */
-	private boolean IsForDistribution = false;
+	protected boolean IsForDistribution = false;
 	
 
 	/** Application build configuration */
@@ -504,10 +506,10 @@ public class GameActivity extends $${gameActivitySuperClass}$$ implements Surfac
 	private int DesiredHolderHeight = 0;
 
 	/** Discovered Vulkan Version and Level from getSystemAvailableFeatures() */
-	private int VulkanVersion = 0;
-	private int VulkanLevel = 0;
-	private String VulkanVersionString = "0.0.0";
-	private boolean bSupportsVulkan = false;
+	public int VulkanVersion = 0;
+	public int VulkanLevel = 0;
+	public String VulkanVersionString = "0.0.0";
+	public boolean bSupportsVulkan = false;
 
 	/** OpenGL information */
 	public String OpenGLVendor = "";
@@ -543,6 +545,14 @@ public class GameActivity extends $${gameActivitySuperClass}$$ implements Surfac
 	}
 	private EAlertDialogType CurrentDialogType = EAlertDialogType.None;
 	
+	public GameActivity()
+	{
+		Log.debug( "GameActivityForMakeAAR constructor called with _gameActivityInstance=" + _gameActivityInstance + ", activityContext=" + activityContext + ", this=" + this + ", isStandalone=" + isStandalone);
+
+		assert (_gameActivityInstance == null);
+		_gameActivityInstance = this;
+	}
+	
 	public boolean IsInVRMode()
 	{
 		return IsInVRMode;
@@ -551,7 +561,12 @@ public class GameActivity extends $${gameActivitySuperClass}$$ implements Surfac
 	/** Access singleton activity for game. **/
 	public static GameActivity Get()
 	{
-		return _activity;
+		return _gameActivityInstance;
+	}
+
+	public static SimpleContextWrapper GetCurrentActivityContext()
+	{
+		return _gameActivityInstance.activityContext;
 	}
 	
 	/**
@@ -588,9 +603,9 @@ public class GameActivity extends $${gameActivitySuperClass}$$ implements Surfac
 	public int getDeviceDefaultOrientation() 
 	{
 		// WindowManager windowManager =  (WindowManager) getSystemService(WINDOW_SERVICE);
-		WindowManager windowManager =  getWindowManager();
+		WindowManager windowManager = activityContext.getWindowManager();
 
-		Configuration config = getResources().getConfiguration();
+		Configuration config = activityContext.getResources().getConfiguration();
 
 		int rotation = windowManager.getDefaultDisplay().getRotation();
 		if ( ((rotation == android.view.Surface.ROTATION_0 || rotation == android.view.Surface.ROTATION_180) &&
@@ -626,7 +641,7 @@ public class GameActivity extends $${gameActivitySuperClass}$$ implements Surfac
 	private int getResourceId(String VariableName, String ResourceName, String PackageName)
 	{
 		try {
-			return getResources().getIdentifier(VariableName, ResourceName, PackageName);
+			return activityContext.getResources().getIdentifier(VariableName, ResourceName, PackageName);
 		}
 		catch (Exception e) {
 			e.printStackTrace();
@@ -640,9 +655,10 @@ public class GameActivity extends $${gameActivitySuperClass}$$ implements Surfac
 		return (resourceId < 1) ? DefaultString : getString(resourceId);
 	}
 
-	private String getProductName()
+	public String getProductName()
 	{
-		String ProductName = Settings.Global.getString(getContentResolver(), "default_device_name");
+	
+		String ProductName = Settings.Global.getString(activityContext.getApplicationContext().getContentResolver(), "default_device_name");
 		if(ProductName == null)
 		{
 			ProductName="";
@@ -699,7 +715,7 @@ public class GameActivity extends $${gameActivitySuperClass}$$ implements Surfac
 
     public boolean AndroidThunkJava_IsPackageInstalled(String packagename)
 	{
-        List<String> packages = getInstalledPackages(getApplicationContext());
+        List<String> packages = getInstalledPackages(activityContext.getApplicationContext());
         return packages.contains(packagename);
     }
 
@@ -1802,7 +1818,7 @@ public class GameActivity extends $${gameActivitySuperClass}$$ implements Surfac
 		int navigationBarHeight = 0;
 		if (Build.VERSION.SDK_INT >= 21)
 		{
-			android.content.res.Resources resources = getApplicationContext().getResources();
+			android.content.res.Resources resources = activityContext.getApplicationContext().getResources();
 			int resourceId = resources.getIdentifier("status_bar_height", "dimen", "android");
 			if (resourceId > 0) {
 				statusBarHeight += resources.getDimensionPixelSize(resourceId);
@@ -1816,7 +1832,7 @@ public class GameActivity extends $${gameActivitySuperClass}$$ implements Surfac
 		variables.put("navigationBarHeight", Integer.toString(navigationBarHeight));
 
 		// store the screen resolution (adjusted for immersive mode)
-		android.view.Display display = getWindowManager().getDefaultDisplay();
+		android.view.Display display = activityContext.getWindowManager().getDefaultDisplay();
 		android.graphics.Point displaySize = new android.graphics.Point();
 		if (ShouldHideUI)
 		{
@@ -2401,7 +2417,8 @@ public class GameActivity extends $${gameActivitySuperClass}$$ implements Surfac
 	@Override
 	public void onCreate(Bundle savedInstanceState)
 	{
-		_activity = this;
+		assert(_gameActivityInstance != null);
+		activityContext = new SimpleContextWrapper(this);
 		InternalFilesDir = getFilesDir().getAbsolutePath() + "/";
 		ExternalFilesDir = getExternalFilesDir(null).getAbsolutePath() + "/";
 		//$${gameActivityOnCreateBeginningAdditions}$$
@@ -2455,7 +2472,7 @@ public class GameActivity extends $${gameActivitySuperClass}$$ implements Surfac
 				SplashScreenLaunch = true;
 				try {
 					// try to get the splash theme (can't use R.style.UnrealSplashTheme since we don't know the package name until runtime)
-					int SplashThemeId = getResources().getIdentifier("UnrealSplashTheme", "style", getPackageName());
+					int SplashThemeId = activityContext.getResources().getIdentifier("UnrealSplashTheme", "style", getPackageName());
 					mSplashDialog = new Dialog(this, SplashThemeId);
 					mSplashDialog.setCancelable(false);
 
@@ -2481,14 +2498,14 @@ public class GameActivity extends $${gameActivitySuperClass}$$ implements Surfac
 					e.printStackTrace();
 				}
 				try {
-					noActionAnimID = getResources().getIdentifier("noaction", "anim", getPackageName());
+					noActionAnimID = activityContext.getResources().getIdentifier("noaction", "anim", getPackageName());
 				}
 				catch (Exception e) {
 					e.printStackTrace();
 				}
 			}
 			try {
-				noActionAnimID = getResources().getIdentifier("noaction", "anim", getPackageName());
+				noActionAnimID = activityContext.getResources().getIdentifier("noaction", "anim", getPackageName());
 			}
 			catch (Exception e) {
 				e.printStackTrace();
@@ -2503,16 +2520,16 @@ public class GameActivity extends $${gameActivitySuperClass}$$ implements Surfac
 			if (UseDisplayCutout)
 			{
 				// will not be true if not Android Pie or later
-	            WindowManager.LayoutParams params = getWindow().getAttributes();
+	            WindowManager.LayoutParams params = activityContext.getWindow().getAttributes();
 		        params.layoutInDisplayCutoutMode = WindowManager.LayoutParams.LAYOUT_IN_DISPLAY_CUTOUT_MODE_SHORT_EDGES;
-			    getWindow().setAttributes(params);
+			  	activityContext.getWindow().setAttributes(params);
 			}
 		}
 
 		// Retrieve version code, name, and target SDK
 		try 
 		{
-			PackageInfo packageInfo = getPackageManager().getPackageInfo(getPackageName(), 0);
+			PackageInfo packageInfo = activityContext.getPackageManager().getPackageInfo(getPackageName(), 0);
 			VersionCode = packageInfo.versionCode;
 			VersionName = packageInfo.versionName;
 			targetSdkVersion = packageInfo.applicationInfo.targetSdkVersion;
@@ -2556,7 +2573,7 @@ public class GameActivity extends $${gameActivitySuperClass}$$ implements Surfac
 			{
 				int ProcessMemory = 0;
 
-				ActivityManager activityManager = (ActivityManager)_activity.getSystemService(Context.ACTIVITY_SERVICE);
+				ActivityManager activityManager = (ActivityManager) activityContext.getSystemService(Context.ACTIVITY_SERVICE);
 				int pid = android.os.Process.myPid();
 				int pids[] = new int[] { pid };
 				android.os.Debug.MemoryInfo[] memoryInfo = activityManager.getProcessMemoryInfo(pids);
@@ -2606,16 +2623,17 @@ public class GameActivity extends $${gameActivitySuperClass}$$ implements Surfac
 					//$${gameActivityMemStatAdditions}$$
 					Log.debug(memLog);
 				}
-				synchronized(_activity)
+				synchronized(activityContext)
 				{
+					Activity localActivity = SimpleContextWrapper.getActivityOfContext(activityContext);
 					if (ProcessMemory > 0)
 					{
-						_activity.UsedMemory = ProcessMemory;
+						UsedMemory = ProcessMemory;
 					}
 					
-					if (_activity.memoryHandler != null && memstatLoggingInterval > 0)
+					if (memoryHandler != null && memstatLoggingInterval > 0)
 					{
-						_activity.memoryHandler.postDelayed(this, memstatLoggingInterval);
+						memoryHandler.postDelayed(this, memstatLoggingInterval);
 					}
 				}
 			}
@@ -2628,8 +2646,8 @@ public class GameActivity extends $${gameActivitySuperClass}$$ implements Surfac
 		// layout required by popups, e.g ads, native controls
 		MarginLayoutParams params = new MarginLayoutParams(LayoutParams.WRAP_CONTENT, LayoutParams.WRAP_CONTENT);
 		params.setMargins(0, 0, 0, 0);
-		activityLayout = new LinearLayout(_activity);
-		_activity.setContentView(activityLayout, params);
+		activityLayout = new LinearLayout(activityContext);
+		SimpleContextWrapper.getActivityOfContext(activityContext).setContentView(activityLayout, params);
 
 /*
 		// Turn on and unlock screen.. Assumption is that this
@@ -2677,7 +2695,7 @@ public class GameActivity extends $${gameActivitySuperClass}$$ implements Surfac
 		// Look for Vulkan support if Nougat or later
 		if (ANDROID_BUILD_VERSION >= 24)
 		{
-			FeatureInfo[] features = getPackageManager().getSystemAvailableFeatures();
+			FeatureInfo[] features = activityContext.getPackageManager().getSystemAvailableFeatures();
 			for (FeatureInfo feature : features) {
 				if (feature.name != null)
 				{
@@ -2728,7 +2746,7 @@ public class GameActivity extends $${gameActivitySuperClass}$$ implements Surfac
 			boolean bForceLandscape = false;
 
 			// check for a Google TV by checking system feature support
-			if (getPackageManager().hasSystemFeature("com.google.android.tv"))
+			if (activityContext.getPackageManager().hasSystemFeature("com.google.android.tv"))
 			{
 				Log.debug( "Detected Google TV, will default to landscape" );
 				bForceLandscape = true;
@@ -2771,7 +2789,7 @@ public class GameActivity extends $${gameActivitySuperClass}$$ implements Surfac
 			if (bForceLandscape)
 			{
 				Log.debug( "Setting screen orientation to landscape because we have detected landscape device" );
-				_activity.setRequestedOrientation( android.content.pm.ActivityInfo.SCREEN_ORIENTATION_LANDSCAPE );
+				SimpleContextWrapper.getActivityOfContext(activityContext).setRequestedOrientation( android.content.pm.ActivityInfo.SCREEN_ORIENTATION_LANDSCAPE );
 			}
 		}
 
@@ -2781,7 +2799,7 @@ public class GameActivity extends $${gameActivitySuperClass}$$ implements Surfac
 		ProjectName = ProjectName.substring(ProjectName.lastIndexOf('.') + 1);
 		String AppType = "";
 		try {
-			ApplicationInfo ai = getPackageManager().getApplicationInfo(getPackageName(), PackageManager.GET_META_DATA);
+			ApplicationInfo ai = activityContext.getPackageManager().getApplicationInfo(getPackageName(), PackageManager.GET_META_DATA);
 			Bundle bundle = ai.metaData;
 			_bundle = bundle;
 
@@ -2972,7 +2990,7 @@ public class GameActivity extends $${gameActivitySuperClass}$$ implements Surfac
 		nativeSetGlobalActivity(UseExternalFilesDir, PublicLogFiles, getFilesDir().getPath(), getExternalFilesDir(null).getPath(), PackageDataInsideApkValue==1, APKPath);
 
 		// tell the engine if this is a portrait app
-		nativeSetWindowInfo(getResources().getConfiguration().orientation == Configuration.ORIENTATION_PORTRAIT, DepthBufferPreference, PropagateAlpha);
+		nativeSetWindowInfo(activityContext.getResources().getConfiguration().orientation == Configuration.ORIENTATION_PORTRAIT, DepthBufferPreference, PropagateAlpha);
 
 		// get the full language code, like en-US
 		// note: this may need to be Locale.getDefault().getLanguage()
@@ -2994,7 +3012,7 @@ public class GameActivity extends $${gameActivitySuperClass}$$ implements Surfac
 		int PatchVersion = 0;
 		if (VersionCode > 0)
 		{
-			nativeSetObbInfo(ProjectName, getApplicationContext().getPackageName(), VersionCode, PatchVersion, AppType);
+			nativeSetObbInfo(ProjectName, activityContext.getApplicationContext().getPackageName(), VersionCode, PatchVersion, AppType);
 		}
 		else
 		{
@@ -3149,8 +3167,8 @@ public class GameActivity extends $${gameActivitySuperClass}$$ implements Surfac
 			HasAllFiles = DownloadShim.expansionFilesDelivered(this, VersionCode);
 		}
 
-		containerFrameLayout = new FrameLayout(_activity);
-		virtualKeyboardLayout = new LinearLayout(_activity);
+		containerFrameLayout = new FrameLayout(activityContext);
+		virtualKeyboardLayout = new LinearLayout(activityContext);
 
 		// Need to create our surface view here regardless of if we are going to end up using it
 		getWindow().takeSurface(null);
@@ -3166,7 +3184,7 @@ public class GameActivity extends $${gameActivitySuperClass}$$ implements Surfac
         mainView.setFocusable( true );
         mainView.setFocusableInTouchMode( true );
 
-        mainDecorView = getWindow().getDecorView();
+        mainDecorView = activityContext.getWindow().getDecorView();
 		mainDecorViewRect = new Rect();
 		
 		// check for display cutouts to set safe zone
@@ -3184,7 +3202,7 @@ public class GameActivity extends $${gameActivitySuperClass}$$ implements Surfac
 							Rect fullViewRect = new Rect();
 							mainDecorView.getDrawingRect( mainDecorViewRect );
 
-							android.view.Display display = getWindowManager().getDefaultDisplay();
+							android.view.Display display = activityContext.getWindowManager().getDefaultDisplay();
 							android.graphics.Point displaySize = new android.graphics.Point();
 							display.getRealSize(displaySize);
 
@@ -3362,7 +3380,7 @@ public class GameActivity extends $${gameActivitySuperClass}$$ implements Surfac
 	{
 		if(android.os.Build.VERSION.SDK_INT >= 19) {
 			try {
-				View decorView = getWindow().getDecorView(); 
+				View decorView = activityContext.getWindow().getDecorView(); 
 
 				Log.debug("=== Restoring Transparent Bars ===");
 				// Clear the flag and then restore it
@@ -3421,13 +3439,13 @@ public class GameActivity extends $${gameActivitySuperClass}$$ implements Surfac
 		}
 
 		// invalidate window cache
-		nativeSetWindowInfo(getResources().getConfiguration().orientation == Configuration.ORIENTATION_PORTRAIT, DepthBufferPreference, PropagateAlpha);
+		nativeSetWindowInfo(activityContext.getResources().getConfiguration().orientation == Configuration.ORIENTATION_PORTRAIT, DepthBufferPreference, PropagateAlpha);
 
 		// only do this on KitKat and above
 		if (ShouldHideUI)
 		{ 
 			restoreTransparentBars();
-			View decorView = getWindow().getDecorView(); 
+			View decorView = activityContext.getWindow().getDecorView(); 
 
 			decorView.setOnSystemUiVisibilityChangeListener(new View.OnSystemUiVisibilityChangeListener() {
 				@Override 
@@ -3452,7 +3470,7 @@ public class GameActivity extends $${gameActivitySuperClass}$$ implements Surfac
 		if (UseDisplayCutout)
 		{
 			// only on Android Pie and later
-			WindowManager.LayoutParams params = getWindow().getAttributes();
+			WindowManager.LayoutParams params = activityContext.getWindow().getAttributes();
 			if (params.layoutInDisplayCutoutMode != WindowManager.LayoutParams.LAYOUT_IN_DISPLAY_CUTOUT_MODE_SHORT_EDGES)
 			{
 				params.layoutInDisplayCutoutMode = WindowManager.LayoutParams.LAYOUT_IN_DISPLAY_CUTOUT_MODE_SHORT_EDGES;
@@ -3477,7 +3495,7 @@ public class GameActivity extends $${gameActivitySuperClass}$$ implements Surfac
 				@Override
 				public void run() {
 					Log.debug("==============> Starting activity to check files and download if required");
-					Intent intent = new Intent(_activity, DownloadShim.GetDownloaderType());
+					Intent intent = new Intent(activityContext, DownloadShim.GetDownloaderType());
 					intent.addFlags(Intent.FLAG_ACTIVITY_NO_ANIMATION);
 					startActivityForResult(intent, DOWNLOAD_ACTIVITY_ID);
 					if (noActionAnimID != -1) {
@@ -3518,7 +3536,7 @@ public class GameActivity extends $${gameActivitySuperClass}$$ implements Surfac
 		if (ShouldHideUI)
 		{ 
 			restoreTransparentBars();
-			View decorView = getWindow().getDecorView(); 
+			View decorView = activityContext.getWindow().getDecorView(); 
 
 			decorView.setOnSystemUiVisibilityChangeListener(new View.OnSystemUiVisibilityChangeListener() {
 				@Override 
@@ -3542,7 +3560,7 @@ public class GameActivity extends $${gameActivitySuperClass}$$ implements Surfac
 		// only true if Android Pie or later
 		if (UseDisplayCutout)
 		{
-			WindowManager.LayoutParams params = getWindow().getAttributes();
+			WindowManager.LayoutParams params = activityContext.getWindow().getAttributes();
 			if (params.layoutInDisplayCutoutMode != WindowManager.LayoutParams.LAYOUT_IN_DISPLAY_CUTOUT_MODE_SHORT_EDGES)
 			{
 				params.layoutInDisplayCutoutMode = WindowManager.LayoutParams.LAYOUT_IN_DISPLAY_CUTOUT_MODE_SHORT_EDGES;
@@ -3560,7 +3578,7 @@ public class GameActivity extends $${gameActivitySuperClass}$$ implements Surfac
 		{
 			final boolean bWarning = bForceGameEndWithWarning;
 
-			_activity.runOnUiThread(new Runnable()
+			activityContext.runOnUiThread(new Runnable()
 			{
 				@Override
 				public void run()
@@ -3570,7 +3588,7 @@ public class GameActivity extends $${gameActivitySuperClass}$$ implements Surfac
 						return;
 					}
 
-					AlertDialog.Builder dialog = new AlertDialog.Builder(_activity);
+					AlertDialog.Builder dialog = new AlertDialog.Builder(activityContext);
 					dialog.setCancelable(false);
 					dialog.setTitle(ForceExitCaption);
 					dialog.setMessage(ForceExitMessage);
@@ -3678,7 +3696,7 @@ public class GameActivity extends $${gameActivitySuperClass}$$ implements Surfac
 		{
 			//	If an AlertDialog is showing when the application is paused, it can cause our main window to be terminated
 			//	Hide the dialog here. It will be shown again via AndroidThunkJava_ShowHiddenAlertDialog called from native code
-			_activity.runOnUiThread(new Runnable()
+			activityContext.runOnUiThread(new Runnable()
 			{
 				public void run()
 				{
@@ -3945,7 +3963,7 @@ public class GameActivity extends $${gameActivitySuperClass}$$ implements Surfac
 			Log.debug("==============> [JAVA] AndroidThunkJava_ShowHiddenAlertDialog() - Showing " + CurrentDialogType);
 		
 			//	If an AlertDialog was showing onPause and we hid it, show it again
-			_activity.runOnUiThread(new Runnable()
+			activityContext.runOnUiThread(new Runnable()
 			{
 				public void run()
 				{
@@ -3976,25 +3994,25 @@ public class GameActivity extends $${gameActivitySuperClass}$$ implements Surfac
 		bKeepScreenOn = Enable;
 		if (Enable)
 		{
-			_activity.runOnUiThread(new Runnable()
+			activityContext.runOnUiThread(new Runnable()
 			{
 				@Override
 				public void run()
 				{
 					Log.debug("==============> [JAVA] AndroidThunkJava_KeepScreenOn(true) - Disabled screen saver");
-					_activity.getWindow().addFlags(WindowManager.LayoutParams.FLAG_KEEP_SCREEN_ON);
+					activityContext.getWindow().addFlags(WindowManager.LayoutParams.FLAG_KEEP_SCREEN_ON);
 				}
 			});
 		}
 		else
 		{
-			_activity.runOnUiThread(new Runnable()
+			activityContext.runOnUiThread(new Runnable()
 			{
 				@Override
 				public void run()
 				{
 					Log.debug("==============> [JAVA] AndroidThunkJava_KeepScreenOn(false) - Enabled screen saver");
-					_activity.getWindow().clearFlags(WindowManager.LayoutParams.FLAG_KEEP_SCREEN_ON);
+					activityContext.getWindow().clearFlags(WindowManager.LayoutParams.FLAG_KEEP_SCREEN_ON);
 				}
 			});
 		}
@@ -4009,25 +4027,25 @@ public class GameActivity extends $${gameActivitySuperClass}$$ implements Surfac
 	{
 		if (bDisable)
 		{
-			_activity.runOnUiThread(new Runnable()
+			activityContext.runOnUiThread(new Runnable()
 			{
 				@Override
 				public void run()
 				{
 					Log.debug("==============> [JAVA] AndroidThunkJava_DisableScreenCapture(true) - Disabled screen captures");
-					_activity.getWindow().addFlags(WindowManager.LayoutParams.FLAG_SECURE);
+					activityContext.getWindow().addFlags(WindowManager.LayoutParams.FLAG_SECURE);
 				}
 			});
 		}
 		else
 		{
-			_activity.runOnUiThread(new Runnable()
+			activityContext.runOnUiThread(new Runnable()
 			{
 				@Override
 				public void run()
 				{
 					Log.debug("==============> [JAVA] AndroidThunkJava_DisableScreenCapture(false) - Enabled screen captures");
-					_activity.getWindow().clearFlags(WindowManager.LayoutParams.FLAG_SECURE);
+					activityContext.getWindow().clearFlags(WindowManager.LayoutParams.FLAG_SECURE);
 				}
 			});
 		}
@@ -4070,7 +4088,7 @@ public class GameActivity extends $${gameActivitySuperClass}$$ implements Surfac
 	{
 		if (deviceVibrator != null)
 		{
-			_activity.runOnUiThread(new VibrateRunnable(Intensity, (long)Duration, deviceVibrator));
+			activityContext.runOnUiThread(new VibrateRunnable(Intensity, (long)Duration, deviceVibrator));
 		}
 	}
 
@@ -4087,7 +4105,7 @@ public class GameActivity extends $${gameActivitySuperClass}$$ implements Surfac
 		consoleInputBox.setHistoryEnd();
 
 		consoleAlert.setMessage("[Available texture formats: " + Formats + "]");
-		_activity.runOnUiThread(new Runnable()
+		activityContext.runOnUiThread(new Runnable()
 		{
 			public void run()
 			{
@@ -4110,7 +4128,7 @@ public class GameActivity extends $${gameActivitySuperClass}$$ implements Surfac
 			return;
 		}
 
-		_activity.runOnUiThread(new Runnable()
+		activityContext.runOnUiThread(new Runnable()
 		{
 			public void run()
 			{
@@ -4138,7 +4156,7 @@ public class GameActivity extends $${gameActivitySuperClass}$$ implements Surfac
 		final String uiLabel = inLabel;
 		final String uiContents = inContents;
 
-		_activity.runOnUiThread(new Runnable()
+		activityContext.runOnUiThread(new Runnable()
 		{
 			public void run()
 			{
@@ -4352,7 +4370,7 @@ public class GameActivity extends $${gameActivitySuperClass}$$ implements Surfac
 	private boolean AndroidThunkJava_isIntentActionAvailable(String intentAction)
 	{
 		Intent intent = new Intent(intentAction, null);
-		List<android.content.pm.ResolveInfo> activities = getPackageManager().queryIntentActivities(intent, PackageManager.MATCH_DEFAULT_ONLY);
+		List<android.content.pm.ResolveInfo> activities = activityContext.getPackageManager().queryIntentActivities(intent, PackageManager.MATCH_DEFAULT_ONLY);
 	    return activities != null && activities.size() > 0;
 	}
 
@@ -4399,7 +4417,7 @@ public class GameActivity extends $${gameActivitySuperClass}$$ implements Surfac
 			BrowserIntent.addFlags(Intent.FLAG_ACTIVITY_NEW_TASK | Intent.FLAG_ACTIVITY_MULTIPLE_TASK);
 			
 			// make sure there is a web browser to handle the URL before trying to start activity (or may crash!)
-			if (BrowserIntent.resolveActivity(getPackageManager()) != null)
+			if (BrowserIntent.resolveActivity(activityContext.getPackageManager()) != null)
 			{
 				Log.debug("[JAVA} AndroidThunkJava_LaunchURL: Starting activity");
 				startActivity(BrowserIntent);
@@ -4418,7 +4436,7 @@ public class GameActivity extends $${gameActivitySuperClass}$$ implements Surfac
 	public String AndroidThunkJava_GetAndroidId()
 	{
 		try {
-			String androidId = Secure.getString(getApplicationContext().getContentResolver(), Secure.ANDROID_ID);
+			String androidId = Secure.getString(activityContext.getApplicationContext().getContentResolver(), Secure.ANDROID_ID);
 			return androidId;
 		} catch (Exception e) {
 			Log.debug("GetAndroidId failed: " + e.getMessage());
@@ -4523,7 +4541,7 @@ public class GameActivity extends $${gameActivitySuperClass}$$ implements Surfac
 			if (bIsFile)
 			{
 				// Note that the ".fileprovider" suffix is hardcoded
-				String Authority = getApplicationContext().getPackageName() + ".fileprovider";
+				String Authority = activityContext.getApplicationContext().getPackageName() + ".fileprovider";
 				Uri FileUri = FileProvider.getUriForFile(this, Authority, SharedFile);
 				
 				SendIntent.putExtra(Intent.EXTRA_STREAM, FileUri);
@@ -4692,6 +4710,11 @@ public class GameActivity extends $${gameActivitySuperClass}$$ implements Surfac
 		return AssetManagerReference;
 	}
 
+	public static void setOBBInAPK(boolean inAPK)
+	{
+		PackageDataInsideApkValue = inAPK ? 1 : 0;
+	}
+
 	public static boolean isOBBInAPK()
 	{
 		Log.debug("Asking if osOBBInAPK? " + (PackageDataInsideApkValue == 1));
@@ -4716,19 +4739,19 @@ public class GameActivity extends $${gameActivitySuperClass}$$ implements Surfac
 	public void AndroidThunkJava_SetOrientation(int value)
 	{
 		final int orientation = value;
-		_activity.runOnUiThread(new Runnable()
+		activityContext.runOnUiThread(new Runnable()
 		{
 			@Override
 			public void run()
 			{
-				_activity.setRequestedOrientation( orientation );
+				SimpleContextWrapper.getActivityOfContext(activityContext).setRequestedOrientation( orientation );
 			}
 		});
 	}
 
 	public void AndroidThunkJava_SetCellularPreference(int value)
 	{
-		Context context = getApplicationContext();
+		Context context = activityContext.getApplicationContext();
 		SharedPreferences preferences = context.getSharedPreferences("CellularNetworkPreferences", context.MODE_PRIVATE);
 		SharedPreferences.Editor editor = preferences.edit();
 		editor.putInt("AllowCellular", value);
@@ -4737,7 +4760,7 @@ public class GameActivity extends $${gameActivitySuperClass}$$ implements Surfac
 
 	public int AndroidThunkJava_GetCellularPreference()
 	{
-		Context context = getApplicationContext();
+		Context context = activityContext.getApplicationContext();
 		SharedPreferences preferences = context.getSharedPreferences("CellularNetworkPreferences", context.MODE_PRIVATE);
 		return preferences.getInt("AllowCellular", 0 );
 	}
@@ -4745,7 +4768,7 @@ public class GameActivity extends $${gameActivitySuperClass}$$ implements Surfac
 	// call back into native code from the Java UI thread, initializing any available VR HMD modules
 	public void AndroidThunkJava_InitHMDs()
 	{
-		_activity.runOnUiThread(new Runnable()
+		activityContext.runOnUiThread(new Runnable()
 		{
 			@Override
 			public void run()
@@ -4805,7 +4828,7 @@ public class GameActivity extends $${gameActivitySuperClass}$$ implements Surfac
 		{
 			bTriggeredQuery = true;
 
-			_activity.runOnUiThread(new Runnable()
+			activityContext.runOnUiThread(new Runnable()
 			{
 				@Override
 				public void run()
@@ -4972,14 +4995,14 @@ public class GameActivity extends $${gameActivitySuperClass}$$ implements Surfac
 		final String messageFinal = message;
 		final boolean bHorizontalFinal = bHorizontal;
 		final int maxValueFinal = maxValue;
-		_activity.runOnUiThread(new Runnable()
+		activityContext.runOnUiThread(new Runnable()
 		{
 			@Override
 			public void run()
 			{
 				if (bShowFinal)
 				{
-					mProgressDialog = new ProgressDialog(_activity);
+					mProgressDialog = new ProgressDialog(activityContext);
 					mProgressDialog.setMessage(messageFinal);
 					mProgressDialog.setProgress(0);
 					mProgressDialog.setMax(maxValueFinal);
@@ -5005,7 +5028,7 @@ public class GameActivity extends $${gameActivitySuperClass}$$ implements Surfac
 			return;
 		}
 		final int valueFinal = value;
-		_activity.runOnUiThread(new Runnable()
+		activityContext.runOnUiThread(new Runnable()
 		{
 			@Override
 			public void run()
@@ -5137,7 +5160,7 @@ public class GameActivity extends $${gameActivitySuperClass}$$ implements Surfac
 
 	private ArrayList<Integer> LocalNotificationGetIDList()
 	{
-		SharedPreferences preferences = getApplicationContext().getSharedPreferences("LocalNotificationPreferences", MODE_PRIVATE);
+		SharedPreferences preferences = activityContext.getApplicationContext().getSharedPreferences("LocalNotificationPreferences", MODE_PRIVATE);
 		SharedPreferences.Editor editor = preferences.edit();
 		String notificationIDs = preferences.getString("notificationIDs", "");
 		ArrayList<Integer> iParts = new ArrayList<Integer>();
@@ -5385,7 +5408,7 @@ public class GameActivity extends $${gameActivitySuperClass}$$ implements Surfac
 		}
 
 		// clear them all
-		SharedPreferences preferences = getApplicationContext().getSharedPreferences("LocalNotificationPreferences", MODE_PRIVATE);
+		SharedPreferences preferences = activityContext.getApplicationContext().getSharedPreferences("LocalNotificationPreferences", MODE_PRIVATE);
 		SharedPreferences.Editor editor = preferences.edit();
 		editor.putString("notificationIDs", "");
 		editor.putString("notificationDetails", "{}");
@@ -5523,7 +5546,7 @@ public class GameActivity extends $${gameActivitySuperClass}$$ implements Surfac
 	{
 		if (vibrator != null)
 		{
-			_activity.runOnUiThread(new VibrateRunnable(intensity, duration, vibrator));
+			activityContext.runOnUiThread(new VibrateRunnable(intensity, duration, vibrator));
 			return true;
 		}
 		return false;
@@ -5578,7 +5601,7 @@ public class GameActivity extends $${gameActivitySuperClass}$$ implements Surfac
 			DesiredHolderHeight > 0 &&
 			MySurfaceView != null)
 		{
-			_activity.runOnUiThread(new Runnable()
+			activityContext.runOnUiThread(new Runnable()
 			{
 				@Override
 				public void run()
@@ -5593,7 +5616,7 @@ public class GameActivity extends $${gameActivitySuperClass}$$ implements Surfac
 	{
 		boolean bPluginEnabled = false;
 //$${gameActivityAllowedRemoteNotificationsAdditions}$$
-		return bPluginEnabled && NotificationManagerCompat.from(getApplicationContext()).areNotificationsEnabled();
+		return bPluginEnabled && NotificationManagerCompat.from(activityContext.getApplicationContext()).areNotificationsEnabled();
 	}
 
 	public void AndroidThunkJava_SetDesiredViewSize(int width, int height)
@@ -5609,7 +5632,7 @@ public class GameActivity extends $${gameActivitySuperClass}$$ implements Surfac
 
 		if(bUseSurfaceView && MySurfaceView != null)
 		{
-			_activity.runOnUiThread(new Runnable()
+			activityContext.runOnUiThread(new Runnable()
 			{
 				@Override
 				public void run()
@@ -5652,7 +5675,7 @@ public class GameActivity extends $${gameActivitySuperClass}$$ implements Surfac
 	{
 		if(ANDROID_BUILD_VERSION >= 24)
 		{
-			WindowManager windowManager =  getWindowManager();
+			WindowManager windowManager = activityContext.getWindowManager();
 			Display display = windowManager.getDefaultDisplay();
 			Display.Mode currentmode = display.getMode();
 			Display.Mode[] modes = display.getSupportedModes();
@@ -5690,7 +5713,7 @@ public class GameActivity extends $${gameActivitySuperClass}$$ implements Surfac
 		if (ANDROID_BUILD_VERSION >= 30)
 		{
 			final int FinalRefreshRate = RefreshRate;
-			_activity.runOnUiThread(new Runnable()
+			activityContext.runOnUiThread(new Runnable()
 			{
 				@Override
 				public void run()
@@ -5712,7 +5735,7 @@ public class GameActivity extends $${gameActivitySuperClass}$$ implements Surfac
 		}
 		else if(ANDROID_BUILD_VERSION >= 24)
 		{
-			WindowManager windowManager =  getWindowManager();
+			WindowManager windowManager = activityContext.getWindowManager();
 			Display display = windowManager.getDefaultDisplay();
 			Display.Mode currentmode = display.getMode();
 			int currentmodeid = currentmode.getModeId();
@@ -5727,11 +5750,11 @@ public class GameActivity extends $${gameActivitySuperClass}$$ implements Surfac
 					final int modeid = modes[i].getModeId();
 					if(currentmodeid != modeid)
 					{
-						_activity.runOnUiThread(new Runnable()
+						activityContext.runOnUiThread(new Runnable()
 						{
 							public void run()
 							{
-								Window w = getWindow();
+								Window w = activityContext.getWindow();
 								WindowManager.LayoutParams l = w.getAttributes();
 								l.preferredDisplayModeId = modeid;
 								w.setAttributes(l);
@@ -5755,7 +5778,7 @@ public class GameActivity extends $${gameActivitySuperClass}$$ implements Surfac
 	{
 		if(ANDROID_BUILD_VERSION >= 24)
 		{
-			WindowManager windowManager =  getWindowManager();
+			WindowManager windowManager = activityContext.getWindowManager();
 			Display display = windowManager.getDefaultDisplay();
 			Display.Mode currentmode = display.getMode();
 			return (int)currentmode.getRefreshRate();
@@ -5804,7 +5827,9 @@ public class GameActivity extends $${gameActivitySuperClass}$$ implements Surfac
 
 	public String getConnectionType()
 	{
-		boolean bAirplaneMode = isAirplaneModeOn(getApplicationContext());
+		assert(activityContext != null);
+
+		boolean bAirplaneMode = isAirplaneModeOn(activityContext.getApplicationContext());
 
 		ConnectivityManager cm = (ConnectivityManager)getSystemService(Context.CONNECTIVITY_SERVICE);
 		NetworkInfo activeNetwork = cm.getActiveNetworkInfo();
@@ -5827,7 +5852,9 @@ public class GameActivity extends $${gameActivitySuperClass}$$ implements Surfac
 
 	public int AndroidThunkJava_GetNetworkConnectionType()
 	{
-		boolean bAirplaneMode = isAirplaneModeOn(getApplicationContext());
+		assert(activityContext != null);
+
+		boolean bAirplaneMode = isAirplaneModeOn(activityContext.getApplicationContext());
 
 		ConnectivityManager cm = (ConnectivityManager)getSystemService(Context.CONNECTIVITY_SERVICE);
 		NetworkInfo activeNetwork = cm.getActiveNetworkInfo();
@@ -5902,6 +5929,9 @@ public class GameActivity extends $${gameActivitySuperClass}$$ implements Surfac
 			return 0;
 		}
 
+		assert(activityContext != null);
+
+
 		switch(key)
 		{
 			case "unreal.http.proxy.proxyPort":
@@ -5913,7 +5943,7 @@ public class GameActivity extends $${gameActivitySuperClass}$$ implements Surfac
 				}
 				else
 				{
-					return android.net.Proxy.getPort(getApplicationContext());
+					return android.net.Proxy.getPort(activityContext.getApplicationContext());
 				}
 			}
 			case "android.hardware.vulkan.version":
@@ -5927,9 +5957,9 @@ public class GameActivity extends $${gameActivitySuperClass}$$ implements Surfac
 			case "unreal.getUsedMemory":
 			{
 				int ProcessMemory = 0;
-				synchronized(_activity)
+				synchronized(activityContext)
 				{
-					ProcessMemory = _activity.UsedMemory;
+					ProcessMemory = UsedMemory;
 				}
 				return ProcessMemory;
 			}
@@ -5968,7 +5998,7 @@ public class GameActivity extends $${gameActivitySuperClass}$$ implements Surfac
 		{
 			if (android.os.Build.VERSION.SDK_INT >= 21)
 			{
-				WindowManager windowManager = (WindowManager)getSystemService(WINDOW_SERVICE);
+				WindowManager windowManager = activityContext.getWindowManager();
 				android.view.Display display = windowManager.getDefaultDisplay();
 				return display.getPresentationDeadlineNanos();
 			}
@@ -5979,7 +6009,7 @@ public class GameActivity extends $${gameActivitySuperClass}$$ implements Surfac
 		{
 			if (android.os.Build.VERSION.SDK_INT >= 21)
 			{
-				WindowManager windowManager = (WindowManager)getSystemService(WINDOW_SERVICE);
+				WindowManager windowManager = activityContext.getWindowManager();
 				android.view.Display display = windowManager.getDefaultDisplay();
 				return display.getAppVsyncOffsetNanos();
 			}
@@ -5992,7 +6022,7 @@ public class GameActivity extends $${gameActivitySuperClass}$$ implements Surfac
 	{
 		if (key.equals("unreal.display.getRefreshRate"))
 		{
-			WindowManager windowManager = (WindowManager)getSystemService(WINDOW_SERVICE);
+			WindowManager windowManager = activityContext.getWindowManager();
 			android.view.Display display = windowManager.getDefaultDisplay();
 			float rate = display.getRefreshRate();
 			// make sure it is sane.. some drivers returned crazy values
@@ -6009,7 +6039,7 @@ public class GameActivity extends $${gameActivitySuperClass}$$ implements Surfac
 		}
 		else
 		{
-			return getWindowManager().getDefaultDisplay();
+			return activityContext.getWindowManager().getDefaultDisplay();
 		}
 	}
 
@@ -6018,17 +6048,17 @@ public class GameActivity extends $${gameActivitySuperClass}$$ implements Surfac
 		DisplayMetrics metrics = new DisplayMetrics();
 		if (android.os.Build.VERSION.SDK_INT >= 31)
 		{
-			metrics = getResources().getDisplayMetrics();
+			metrics = activityContext.getResources().getDisplayMetrics();
 		}
 		else
 		if (android.os.Build.VERSION.SDK_INT >= 17)
 		{
-			getWindowManager().getDefaultDisplay().getRealMetrics(metrics);
+			activityContext.getWindowManager().getDefaultDisplay().getRealMetrics(metrics);
 		}
 		else
 		{
 			// note: not available so get what we can
-			getWindowManager().getDefaultDisplay().getMetrics(metrics);
+			activityContext.getWindowManager().getDefaultDisplay().getMetrics(metrics);
 		}
 		return metrics;
 	}
@@ -6130,7 +6160,7 @@ public class GameActivity extends $${gameActivitySuperClass}$$ implements Surfac
 				}
 				else
 				{
-					return android.net.Proxy.getHost(getApplicationContext());
+					return android.net.Proxy.getHost(activityContext.getApplicationContext());
 				}
 			}
 			case "unreal.displaymetrics.dpi":
@@ -6225,13 +6255,13 @@ public class GameActivity extends $${gameActivitySuperClass}$$ implements Surfac
 					Log.warn("SustainedPerformanceMode is unsupported on this device.");
 				}
 			}
-			_activity.runOnUiThread(new Runnable()
+			activityContext.runOnUiThread(new Runnable()
 			{
 				public void run()
 				{
 					try
 					{
-						android.view.Window ActivityWindow = _activity.getWindow();
+						android.view.Window ActivityWindow = activityContext.getWindow();
 						Method m = android.view.Window.class.getMethod("setSustainedPerformanceMode",Boolean.TYPE);
 						m.invoke(ActivityWindow, bEnable);
 					}
@@ -6690,7 +6720,7 @@ public class GameActivity extends $${gameActivitySuperClass}$$ implements Surfac
 
 		if (intentAction.equals(""))
 		{
-			intent = getPackageManager().getLaunchIntentForPackage(packageName);
+			intent = activityContext.getPackageManager().getLaunchIntentForPackage(packageName);
 			if (intent == null)
 			{
 				return false;
@@ -6717,7 +6747,7 @@ public class GameActivity extends $${gameActivitySuperClass}$$ implements Surfac
 			intent.setComponent(new android.content.ComponentName(packageName, intentComponent));
 		}
 
-		if (intent.resolveActivity(getPackageManager()) != null)
+		if (intent.resolveActivity(activityContext.getPackageManager()) != null)
 		{
 			startActivity(intent);
 			AndroidThunkJava_ForceQuit();
@@ -6754,7 +6784,7 @@ public class GameActivity extends $${gameActivitySuperClass}$$ implements Surfac
 
 		if (intentAction.equals(""))
 		{
-			intent = getPackageManager().getLaunchIntentForPackage(packageName);
+			intent = activityContext.getPackageManager().getLaunchIntentForPackage(packageName);
 			if (intent == null)
 			{
 				return false;
@@ -6781,7 +6811,7 @@ public class GameActivity extends $${gameActivitySuperClass}$$ implements Surfac
 			intent.setComponent(new android.content.ComponentName(packageName, intentComponent));
 		}
 
-		getApplicationContext().sendBroadcast(intent);
+		activityContext.getApplicationContext().sendBroadcast(intent);
 
 		if (bExit)
 		{
@@ -6791,9 +6821,9 @@ public class GameActivity extends $${gameActivitySuperClass}$$ implements Surfac
 		return true;
 	}
 
-	private void RestartApplication(String RestartExtra)
+	protected void RestartApplication(String RestartExtra)
 	{
-		Context context = getApplicationContext();
+		Context context = activityContext.getApplicationContext();
 	
 		// get the launch intent
 		PackageManager packageManager = context.getPackageManager();
@@ -6860,13 +6890,13 @@ public class GameActivity extends $${gameActivitySuperClass}$$ implements Surfac
 
 	public int AndroidThunkJava_GetDeviceRotation() 
 	{
-		return getWindowManager().getDefaultDisplay().getRotation();
+		return activityContext.getWindowManager().getDefaultDisplay().getRotation();
 	}
 
 	public int AndroidThunkJava_GetDeviceOrientation()
 	{
 		int currentScreenOrientation = DEVICE_ORIENTATION_UNKNOWN;
-		int rotation = getWindowManager().getDefaultDisplay().getRotation();
+		int rotation = activityContext.getWindowManager().getDefaultDisplay().getRotation();
 		if ( getDeviceDefaultOrientation() == Configuration.ORIENTATION_LANDSCAPE )
 		{
 			switch (rotation)
@@ -6914,7 +6944,7 @@ public class GameActivity extends $${gameActivitySuperClass}$$ implements Surfac
 		}
 		else
 		{
-			Context context = getApplicationContext();
+			Context context = activityContext.getApplicationContext();
 			CookieSyncManager syncManager = CookieSyncManager.createInstance(context);
 			syncManager.sync();
 
@@ -6937,7 +6967,7 @@ public class GameActivity extends $${gameActivitySuperClass}$$ implements Surfac
 	{
 		if (android.os.Build.VERSION.SDK_INT < android.os.Build.VERSION_CODES.LOLLIPOP)
 		{
-			Context context = getApplicationContext();
+			Context context = activityContext.getApplicationContext();
 			CookieSyncManager syncManager = CookieSyncManager.createInstance(context);
 			syncManager.sync();
 		}
@@ -6956,7 +6986,7 @@ public class GameActivity extends $${gameActivitySuperClass}$$ implements Surfac
 	{
 		if (android.os.Build.VERSION.SDK_INT < android.os.Build.VERSION_CODES.LOLLIPOP)
 		{
-			Context context = getApplicationContext();
+			Context context = activityContext.getApplicationContext();
 			CookieSyncManager syncManager = CookieSyncManager.createInstance(context);
 			syncManager.sync();
 
@@ -7044,7 +7074,7 @@ public class GameActivity extends $${gameActivitySuperClass}$$ implements Surfac
 
 	public native void nativeInitHMDs();
 
-	public native void nativeResumeMainInit();
+	public native static void nativeResumeMainInit();
 
 	public native void nativeOnActivityResult(GameActivity activity, int requestCode, int resultCode, Intent data);
 
@@ -7098,6 +7128,6 @@ public class GameActivity extends $${gameActivitySuperClass}$$ implements Surfac
 
 	public native void nativeHandleSensorEvents(float[] tilt, float[] rotation_rate, float[] gravity, float[] acceleration);
 
-	private static ClipboardManager clipboardManager;
+	protected static ClipboardManager clipboardManager;
 
 }
diff --git a/Engine/Build/Android/Java/src/com/epicgames/unreal/SimpleContextWrapper.java b/Engine/Build/Android/Java/src/com/epicgames/unreal/SimpleContextWrapper.java
new file mode 100644
--- /dev/null
+++ b/Engine/Build/Android/Java/src/com/epicgames/unreal/SimpleContextWrapper.java
@@ -0,0 +1,267 @@
+package com.epicgames.unreal;
+
+import android.app.Activity;
+import android.content.Context;
+import android.content.ContextWrapper;
+import android.os.Handler;
+import android.view.View;
+import android.view.ViewGroup;
+import android.view.Window;
+import android.view.WindowManager;
+
+import java.lang.ref.PhantomReference;
+import java.lang.ref.WeakReference;
+import java.util.Arrays;
+import java.util.concurrent.BlockingQueue;
+import java.util.concurrent.LinkedBlockingQueue;
+import java.util.concurrent.ThreadPoolExecutor;
+import java.util.concurrent.TimeUnit;
+
+// valid calls for Context...
+// getResources()
+// getApplicationContext()
+// getPackageManager()
+
+public final class SimpleContextWrapper extends ContextWrapper {
+	final Handler mHandler;
+	final private Thread mUiThread;
+	final public Context mAppActivityContext;
+
+	public static Logger Log = new Logger("UE", "SimpleContextWrapper");
+
+	public static Activity getActivityOfContext(Context context) {
+		if (context == null) {
+			return null;
+		}
+
+		while (context instanceof ContextWrapper) {
+			if (context instanceof Activity) {
+				return (Activity) context;
+			}
+			context = ((ContextWrapper) context).getBaseContext();
+		}
+		return null;
+	}
+
+	public static Context getBaseOfContext(Context context) {
+		if (context == null) {
+			return null;
+		}
+
+		while (context instanceof ContextWrapper) {
+
+			if ( !(((ContextWrapper) context).getBaseContext() instanceof ContextWrapper)) {
+				return ((ContextWrapper) context).getBaseContext();
+			}
+			context = ((ContextWrapper) context).getBaseContext();
+		}
+
+		return null;
+	}
+
+//	public static boolean isSimpleActivityOfContext(Context context) {
+//		if (context == null) {
+//			return false;
+//		}
+//
+//		while (context instanceof ContextWrapper) {
+//			if (context instanceof SimpleContextWrapper) {
+//				return true;
+//			}
+//			context = ((ContextWrapper) context).getBaseContext();
+//		}
+//
+//		return false;
+//	}
+
+	public static ContextWrapper getContextWrapperOfContext(Context context) {
+		if (context == null) {
+			return null;
+		}
+
+		Context searchContext = context;
+		while (searchContext instanceof ContextWrapper) {
+			if (searchContext instanceof SimpleContextWrapper) {
+				return (SimpleContextWrapper)searchContext;
+			}
+			searchContext = ((ContextWrapper) searchContext).getBaseContext();
+		}
+
+
+		searchContext = context;
+		while (searchContext instanceof ContextWrapper) {
+			if (searchContext instanceof Activity) {
+				return (ContextWrapper)searchContext;
+			}
+			searchContext = ((ContextWrapper) searchContext).getBaseContext();
+		}
+
+		while (context instanceof ContextWrapper) {
+			if (context instanceof ContextWrapper) {
+				return (ContextWrapper)context;
+			}
+			context = ((ContextWrapper) context).getBaseContext();
+		}
+
+		return null;
+	}
+
+	public static SimpleContextWrapper getSimpleActivityOfContext(Context context) {
+		if (context == null) {
+			return null;
+		}
+
+		while (context instanceof ContextWrapper) {
+			if (context instanceof SimpleContextWrapper) {
+				return (SimpleContextWrapper)context;
+			}
+			context = ((ContextWrapper) context).getBaseContext();
+		}
+
+		return null;
+	}
+
+
+//	@UiThread
+	public SimpleContextWrapper(SimpleContextWrapper base) {
+		super(base);
+		this.mHandler = base.mHandler;
+		this.mUiThread = base.mUiThread;
+		this.mAppActivityContext = base.mAppActivityContext;
+	}
+
+//	@UiThread
+	public SimpleContextWrapper(Activity base) {
+		super(base);
+		SimpleContextWrapper _base = getSimpleActivityOfContext(base);
+		Log.verbose("UE - SimpleContextWrapper(Activity base),  base = " + base + ", _base= " + _base + ", Thread.currentThread() = " + Thread.currentThread());
+
+		if (_base != null)
+		{
+			this.mHandler = _base.mHandler;
+			this.mUiThread = _base.mUiThread;
+			this.mAppActivityContext = _base;
+		}
+		else
+		{
+			this.mHandler = new Handler();
+			this.mUiThread = Thread.currentThread();
+			this.mAppActivityContext = getActivityOfContext(base);
+		}
+
+	}
+
+//	@UiThread
+	public SimpleContextWrapper(Context base) {
+		super(base);
+
+		SimpleContextWrapper _base = getSimpleActivityOfContext(base);
+		Log.verbose("UE - SimpleContextWrapper(Context base),  base = " + base + ", _base= " + _base + ", Thread.currentThread() = " + Thread.currentThread());
+
+		if (_base != null)
+		{
+			this.mHandler = _base.mHandler;
+			this.mUiThread = _base.mUiThread;
+			this.mAppActivityContext = _base.mAppActivityContext;
+		}
+		else
+		{
+			this.mHandler = new Handler();
+			this.mUiThread = Thread.currentThread();
+			this.mAppActivityContext = getActivityOfContext(base);
+		}
+	}
+
+	public final void runOnUiThread(Runnable action) {
+		Log.verbose("UE - SimpleContextWrapper, trying to runOnUiThread mAppActivityContext = " + mAppActivityContext + ", mUiThread= " + mUiThread + ", Thread.currentThread() = " + Thread.currentThread());
+
+		Activity _activity = mAppActivityContext instanceof Activity ? (Activity) mAppActivityContext : null;
+		if (_activity == null)
+		{
+			_activity = getActivityOfContext(this);
+		}
+		if (_activity == null)
+		{
+			_activity = getActivityOfContext(getBaseOfContext(this));
+		}
+		
+		if (_activity != null)
+		{
+			_activity.runOnUiThread(action);
+		}
+		else 
+		{
+
+			if (Thread.currentThread() != mUiThread) {
+				mHandler.post(action);
+			} else {
+				Log.error("UE - SimpleContextWrapper, trying to runOnUiThread is being BYPASSED and action called on Main Thread, mAppActivityContext = " + mAppActivityContext + ", mUiThread= " + mUiThread + ", Thread.currentThread() = " + Thread.currentThread());
+
+				//action.run();
+				runOnMainThread(action);
+				
+			}
+		}
+	}
+
+	public Window getWindow() {
+		return mAppActivityContext != null ? ((Activity)mAppActivityContext).getWindow() : null;
+	}
+
+	public WindowManager getWindowManager() {
+		return mAppActivityContext != null ?  ((Activity)mAppActivityContext).getWindowManager() : (WindowManager) getSystemService(Context.WINDOW_SERVICE);
+	}
+
+	public void addContentView(View view, ViewGroup.LayoutParams params) {
+		if (mAppActivityContext != null) {
+			((Activity)mAppActivityContext).addContentView(view, params);
+		}
+	}
+
+	public boolean equals(Context obj) {
+		return (this.hashCode() == obj.hashCode());
+	}
+
+
+	/*
+	 * Gets the number of available cores
+	 * (not always the same as the maximum number of cores)
+	 */
+	private static int NUMBER_OF_CORES = Runtime.getRuntime().availableProcessors();
+
+	// Instantiates the queue of Runnables as a LinkedBlockingQueue
+	private final static BlockingQueue<Runnable> workQueue = new LinkedBlockingQueue<Runnable>();
+
+	// Sets the amount of time an idle thread waits before terminating
+	private static final int KEEP_ALIVE_TIME = 1;
+	// Sets the Time Unit to seconds
+	private static final TimeUnit KEEP_ALIVE_TIME_UNIT = TimeUnit.SECONDS;
+	
+	static ThreadPoolExecutor threadPoolExecutor;
+
+	private static void initThreadPoolExtractor() {
+
+		// Creates a thread pool manager
+		threadPoolExecutor = new ThreadPoolExecutor(
+			NUMBER_OF_CORES,       // Initial pool size
+			NUMBER_OF_CORES,       // Max pool size
+			KEEP_ALIVE_TIME,
+			KEEP_ALIVE_TIME_UNIT,
+			workQueue
+		);
+	}
+
+
+	public final static void runOnMainThread (Runnable action){
+
+		if (threadPoolExecutor == null) {
+			initThreadPoolExtractor();
+		}
+		threadPoolExecutor.execute(action);
+
+//		action.run();
+
+//		new Handler(Looper.getMainLooper()).post(action);
+	}
+
+}
\ No newline at end of file
diff --git a/Engine/Build/Android/Java/src/com/epicgames/unreal/WebViewControl.java b/Engine/Build/Android/Java/src/com/epicgames/unreal/WebViewControl.java
--- a/Engine/Build/Android/Java/src/com/epicgames/unreal/WebViewControl.java
+++ b/Engine/Build/Android/Java/src/com/epicgames/unreal/WebViewControl.java
@@ -123,7 +123,7 @@ class WebViewControl
 
 		//GameActivity.Log.debug("WebViewControl width=" + width + ", height=" + height);
 
-		GameActivity._activity.runOnUiThread(new Runnable()
+		GameActivity.Get().runOnUiThread(new Runnable()
 		{
 			@Override
 			public void run()
@@ -131,11 +131,11 @@ class WebViewControl
 				// enable remote debugging if requested and supported by the current platform
 				if (Build.VERSION.SDK_INT >= Build.VERSION_CODES.KITKAT)
 				{
-					WebView.setWebContentsDebuggingEnabled(bEnableRemoteDebugging && !GameActivity._activity.nativeIsShippingBuild());
+					WebView.setWebContentsDebuggingEnabled(bEnableRemoteDebugging && !GameActivity.Get().nativeIsShippingBuild());
 				}
 				
 				// create the WebView
-				webView = new GLWebView(GameActivity._activity);
+				webView = new GLWebView(GameActivity.Get());
 				webView.setWebViewClient(new ViewClient());
 				webView.setWebChromeClient(new ChromeClient());
 				webView.getSettings().setJavaScriptEnabled(true);
@@ -165,7 +165,7 @@ class WebViewControl
 				}
 
 				// Wrap the webview in a layout that will do absolute positioning for us
-				positionLayout = new WebViewPositionLayout(GameActivity._activity, w);
+				positionLayout = new WebViewPositionLayout(GameActivity.Get(), w);
 				ViewGroup.LayoutParams params = new ViewGroup.LayoutParams(ViewGroup.LayoutParams.FILL_PARENT, ViewGroup.LayoutParams.FILL_PARENT);
 
 				positionLayout.addView(webView, params);
@@ -182,7 +182,7 @@ class WebViewControl
 	public void SetAndroid3DBrowser(boolean InIsAndroid3DBrowser)
 	{
 		PendingSetAndroid3DBrowser = (TargetBitmap == null) ? InIsAndroid3DBrowser : false;
-		GameActivity._activity.runOnUiThread(new Runnable()
+		GameActivity.Get().runOnUiThread(new Runnable()
 		{
 			@Override
 			public void run()
@@ -196,7 +196,7 @@ class WebViewControl
 	public void SetVisibility(boolean InIsVisible)
 	{
 		PendingSetVisibility = InIsVisible;
-		GameActivity._activity.runOnUiThread(new Runnable()
+		GameActivity.Get().runOnUiThread(new Runnable()
 		{
 			@Override
 			public void run()
@@ -247,7 +247,7 @@ class WebViewControl
 
 	public void ExecuteJavascript(final String script)
 	{
-		GameActivity._activity.runOnUiThread(new Runnable()
+		GameActivity.Get().runOnUiThread(new Runnable()
 		{
 			@Override
 			public void run()
@@ -269,7 +269,7 @@ class WebViewControl
 	}
 	public void LoadURL(final String url)
 	{
-		GameActivity._activity.runOnUiThread(new Runnable()
+		GameActivity.Get().runOnUiThread(new Runnable()
 		{
 			@Override
 			public void run()
@@ -282,7 +282,7 @@ class WebViewControl
 
 	public void LoadString(final String contents, final String url)
 	{
-		GameActivity._activity.runOnUiThread(new Runnable()
+		GameActivity.Get().runOnUiThread(new Runnable()
 		{
 			@Override
 			public void run()
@@ -295,7 +295,7 @@ class WebViewControl
 
 	public void StopLoad()
 	{
-		GameActivity._activity.runOnUiThread(new Runnable()
+		GameActivity.Get().runOnUiThread(new Runnable()
 		{
 			@Override
 			public void run()
@@ -307,7 +307,7 @@ class WebViewControl
 
 	public void Reload()
 	{
-		GameActivity._activity.runOnUiThread(new Runnable()
+		GameActivity.Get().runOnUiThread(new Runnable()
 		{
 			@Override
 			public void run()
@@ -319,7 +319,7 @@ class WebViewControl
 
 	public void GoBackOrForward(final int Steps)
 	{
-		GameActivity._activity.runOnUiThread(new Runnable()
+		GameActivity.Get().runOnUiThread(new Runnable()
 		{
 			@Override
 			public void run()
@@ -341,7 +341,7 @@ class WebViewControl
 		final float actionX = webView.getLeft() + (x * webView.getWidth());
 		final float actionY = webView.getTop() + (y * webView.getHeight());
 		//GameActivity.Log.debug("SendTouchEvent(event=" + event + ", x=" + x + ", y=" + y + ") = " + actionX + ", " + actionY);
-		GameActivity._activity.runOnUiThread(new Runnable()
+		GameActivity.Get().runOnUiThread(new Runnable()
 		{
 			@Override
 			public void run()
@@ -514,7 +514,7 @@ class WebViewControl
 		final int actionCode = KeyTableList[keycode].key;
 		final int actionMeta = KeyTableList[keycode].meta;
 //		GameActivity.Log.debug("SendTouchEvent(event=" + (bDown ? "Down" : "Up") + ", key=" + keycode + ", code=" + actionCode + ", meta=" + actionMeta);
-		GameActivity._activity.runOnUiThread(new Runnable()
+		GameActivity.Get().runOnUiThread(new Runnable()
 		{
 			@Override
 			public void run()
@@ -529,7 +529,7 @@ class WebViewControl
 	// called from C++ paint event
 	public void Update(final int x, final int y, final int width, final int height)
 	{
-		GameActivity._activity.runOnUiThread(new Runnable()
+		GameActivity.Get().runOnUiThread(new Runnable()
 		{
 			@Override
 			public void run()
@@ -545,7 +545,7 @@ class WebViewControl
 					
 					// add to the activitiy, on top of the SurfaceView
 					ViewGroup.LayoutParams params = new ViewGroup.LayoutParams(ViewGroup.LayoutParams.FILL_PARENT, ViewGroup.LayoutParams.FILL_PARENT);			
-					GameActivity._activity.addContentView(positionLayout, params);
+					GameActivity.Get().addContentView(positionLayout, params);
 					if(!webView.IsAndroid3DBrowser && TargetBitmap == null)
 					{
 						//GameActivity.Log.warn("request focus create");
@@ -627,7 +627,7 @@ class WebViewControl
 
 	public void Close()
 	{
-		GameActivity._activity.runOnUiThread(new Runnable()
+		GameActivity.Get().runOnUiThread(new Runnable()
 		{
 			@Override
 			public void run()
@@ -666,7 +666,7 @@ class WebViewControl
 		
 		// set this here as the size may have been set before the GL resources were created.
 		mBitmapRenderer.setSize(initialWidth, initialHeight);
-		GameActivity._activity.runOnUiThread(new Runnable()
+		GameActivity.Get().runOnUiThread(new Runnable()
 		{
 			@Override
 			public void run()
@@ -1594,7 +1594,7 @@ class WebViewControl
 		
 		// set this here as the size may have been set before the GL resources were created.
 		mOESTextureRenderer.setSize(initialWidth, initialHeight);
-		GameActivity._activity.runOnUiThread(new Runnable()
+		GameActivity.Get().runOnUiThread(new Runnable()
 		{
 			@Override
 			public void run()
@@ -1966,7 +1966,7 @@ class WebViewControl
 		@Override
 		public boolean onCreateWindow(WebView View, boolean isDialog, boolean isUserGesture, Message resultMsg)
 		{
-			WebView newView = new WebView(GameActivity._activity);
+			WebView newView = new WebView(GameActivity.Get());
 			View.addView(newView);
             WebView.WebViewTransport transport = (WebView.WebViewTransport) resultMsg.obj;
             transport.setWebView(newView);
diff --git a/Engine/Plugins/Runtime/AndroidSingleInstanceService/Source/AndroidSingleInstanceService/AndroidSingleInstanceService.Build.cs b/Engine/Plugins/Runtime/AndroidSingleInstanceService/Source/AndroidSingleInstanceService/AndroidSingleInstanceService.Build.cs
--- a/Engine/Plugins/Runtime/AndroidSingleInstanceService/Source/AndroidSingleInstanceService/AndroidSingleInstanceService.Build.cs
+++ b/Engine/Plugins/Runtime/AndroidSingleInstanceService/Source/AndroidSingleInstanceService/AndroidSingleInstanceService.Build.cs
@@ -7,6 +7,8 @@ public class AndroidSingleInstanceService : ModuleRules
 {
 	public AndroidSingleInstanceService(ReadOnlyTargetRules Target) : base(Target)
 	{
+		ShortName = "AndSngInstSvc";
+
 		PCHUsage = ModuleRules.PCHUsageMode.UseExplicitOrSharedPCHs;
 		
 		PublicDependencyModuleNames.AddRange(
diff --git a/Engine/Plugins/Runtime/AndroidSingleInstanceService/Source/AndroidSingleInstanceServiceEditor/AndroidSingleInstanceServiceEditor.Build.cs b/Engine/Plugins/Runtime/AndroidSingleInstanceService/Source/AndroidSingleInstanceServiceEditor/AndroidSingleInstanceServiceEditor.Build.cs
--- a/Engine/Plugins/Runtime/AndroidSingleInstanceService/Source/AndroidSingleInstanceServiceEditor/AndroidSingleInstanceServiceEditor.Build.cs
+++ b/Engine/Plugins/Runtime/AndroidSingleInstanceService/Source/AndroidSingleInstanceServiceEditor/AndroidSingleInstanceServiceEditor.Build.cs
@@ -6,6 +6,8 @@ namespace UnrealBuildTool.Rules
 	{
 		public AndroidSingleInstanceServiceEditor(ReadOnlyTargetRules Target) : base(Target)
 		{
+			ShortName = "AndSngInstSvcEd";
+
 			PrivateDependencyModuleNames.AddRange(
 				new string[] {
 					"Core",
