Hooking libc.so malloc/free/etc into memory tracer on Android

diff --git a/Engine/Build/Android/Prebuilt/ScudoMemoryTrace/Build.bat b/Engine/Build/Android/Prebuilt/ScudoMemoryTrace/Build.bat
new file mode 100755
--- /dev/null
+++ b/Engine/Build/Android/Prebuilt/ScudoMemoryTrace/Build.bat
@@ -0,0 +1,22 @@
+set CMAKEPATH=%ANDROID_HOME%\cmake\3.22.1\bin
+
+call :build_arch arm64-v8a
+call :build_arch x86_64
+
+exit /b
+
+:build_arch
+
+mkdir .build_%1
+pushd .build_%1
+
+%CMAKEPATH%\cmake.exe -GNinja -DCMAKE_BUILD_TYPE=Release -DCMAKE_SYSTEM_NAME=Android -DCMAKE_SYSTEM_VERSION=28 -DCMAKE_ANDROID_ARCH_ABI=%1 -DCMAKE_ANDROID_STL_TYPE=c++_static -DCMAKE_MAKE_PROGRAM=%CMAKEPATH%\ninja.exe ..
+%CMAKEPATH%\ninja.exe
+
+popd
+
+mkdir %1
+move /y .build_%1\libScudoMemoryTrace.so %1
+rmdir .build_%1 /s /q
+
+exit /b
diff --git a/Engine/Build/Android/Prebuilt/ScudoMemoryTrace/CMakeLists.txt b/Engine/Build/Android/Prebuilt/ScudoMemoryTrace/CMakeLists.txt
new file mode 100644
--- /dev/null
+++ b/Engine/Build/Android/Prebuilt/ScudoMemoryTrace/CMakeLists.txt
@@ -0,0 +1,14 @@
+cmake_minimum_required(VERSION 3.22.1)
+project(ScudoMemoryTrace)
+
+message("Using NDK version ${CMAKE_ANDROID_NDK_VERSION}")
+
+add_library(ScudoMemoryTrace SHARED
+  ${CMAKE_SOURCE_DIR}/ScudoMemoryTrace.cpp
+  ${CMAKE_SOURCE_DIR}/ScudoMemoryTrace.h
+)
+
+set_property(TARGET ScudoMemoryTrace PROPERTY CXX_STANDARD 20)
+
+set_target_properties(ScudoMemoryTrace PROPERTIES COMPILE_FLAGS "-fno-exceptions -fno-rtti -fno-unwind-tables -g -O3")
+set_target_properties(ScudoMemoryTrace PROPERTIES LINK_FLAGS_RELEASE "-s")
diff --git a/Engine/Build/Android/Prebuilt/ScudoMemoryTrace/ScudoMemoryTrace.cpp b/Engine/Build/Android/Prebuilt/ScudoMemoryTrace/ScudoMemoryTrace.cpp
new file mode 100644
--- /dev/null
+++ b/Engine/Build/Android/Prebuilt/ScudoMemoryTrace/ScudoMemoryTrace.cpp
@@ -0,0 +1,742 @@
+// Copyright Epic Games, Inc. All Rights Reserved.
+
+// libc.so/scudo memory tracing.
+// Implemented overriding malloc, calloc, realloc, memalign, free, posix_memalign, aligned_alloc, reallocarray via LD_PRELOAD.
+// When modifying this code beware that any allocation from with-in malloc handler will cause recursion,
+// for example thread local will call malloc on the first access, snprintf calls malloc, Android log calls malloc, etc.
+
+#include "ScudoMemoryTrace.h"
+
+// Trace events before libUnreal.so is loaded and sets the trace hook.
+// This requires creating a temporary file as a backing storage to store the events.
+constexpr bool bTraceEventsBeforeLibUnrealIsLoaded = true;
+
+// If bTraceEventsBeforeLibUnrealIsLoaded is true, we record everything as soon as the application starts,
+// but if the hook is never set we need to stop recording to not run over the storage.
+// Also there are other processes besides main one during the app launch, and tracer will be preloaded into them as well.
+constexpr uint32_t MaxEventsBeforeHookIsSet = 1'000'000;
+
+// If bTraceEventsBeforeLibUnrealIsLoaded is true, enable dumping events into a debug file, will use a tmp file instead if not defined (tmp file is not linked to fs).
+constexpr const char* DebugFilePath = nullptr; // "/data/user/0/%PACKAGE_NAME%/files/llmemorytrace_";
+
+// If bTraceEventsBeforeLibUnrealIsLoaded is true, size of in-memory lockfree double buffer used to temporary store events.
+static const size_t EventBufferSize = 16 * 1024;
+
+// Experimental generation of time events for pre-libUnreal.so memory tracing.
+constexpr bool bGenerateTimeEvents = false;
+
+// If bGenerateTimeEvents is true, how often to set timestamps before the marker is set.
+static const uint32_t BeforeHookMarkerSamplePeriod = (4 << 10) - 1;
+
+// On some OS scudo has special hooks available, but they were disabled in Android 15.
+#define USE_SCUDO_HOOKS 0
+
+#include <fcntl.h>
+#include <unistd.h>
+#include <dlfcn.h>
+#include <atomic>
+#include <mutex>
+#include <thread>
+#include <sys/time.h>
+
+#pragma pack(push, 1)
+
+struct FEvent
+{
+	uint64_t Ptr;
+	uint64_t Size;
+	uint32_t Alignment;
+	EScudoEventType Type;
+	uint8_t FrameCount;
+
+	static const uint32_t MaxCallstackFrameCount = 20;
+	uint64_t Frames[MaxCallstackFrameCount];
+
+	constexpr static inline uint32_t GetEmptyEventSize()
+	{
+		return sizeof(FEvent) - MaxCallstackFrameCount * sizeof(uint64_t);
+	}
+
+	inline uint32_t GetEventSize() const
+	{
+		return GetEmptyEventSize() + FrameCount * sizeof(uint64_t);
+	}
+};
+
+#pragma pack(pop)
+
+struct Spinlock
+{
+	std::atomic_flag Flag = {false};
+
+	inline void lock()
+	{
+		while (Flag.test_and_set(std::memory_order_acquire))
+		{
+			while (Flag.test(std::memory_order_relaxed))
+			{
+				std::this_thread::yield();
+			}
+		}
+	}
+
+	inline void unlock()
+	{
+		Flag.clear(std::memory_order_release);
+	}
+};
+
+struct TLSFlag
+{
+	pthread_key_t Key;
+
+	inline void Setup()
+	{
+		pthread_key_create(&Key, nullptr);
+	}
+
+	inline void Set(bool bValue)
+	{
+		pthread_setspecific(Key, bValue ? (const void*)1 : nullptr);
+	}
+
+	inline bool Get()
+	{
+		return pthread_getspecific(Key) != nullptr;
+	}
+};
+
+struct EventBuffer
+{
+	uint8_t Data[EventBufferSize];
+	std::atomic<uint32_t> Head = {0};
+	std::atomic<uint32_t> Tail = {0};
+
+	inline bool Write(const FEvent* Event)
+	{
+		const void* Buffer = (const void*)Event;
+		const uint32_t EventSize = Event->GetEventSize();
+
+		// Acquire space in the buffer
+		uint32_t WritePos = Head.load();
+		do
+		{
+			if (WritePos + EventSize > EventBufferSize)
+			{
+				return false;
+			}
+		}
+		while(Head.compare_exchange_strong(WritePos, WritePos + EventSize) == false);
+
+		memcpy(Data + WritePos, Buffer, EventSize);
+
+		Tail.fetch_add(EventSize);
+		return true;
+	}
+
+	inline bool Read(FEvent* Event, uint32_t* Offset, const uint32_t DataSize)
+	{
+		const uint32_t EmptySize = FEvent::GetEmptyEventSize();
+		if (*Offset + EmptySize > DataSize)
+		{
+			return false;
+		}
+
+		memcpy(Event, Data + *Offset, EmptySize);
+
+		const uint32_t EventSize = Event->GetEventSize();
+		if (*Offset + EventSize > DataSize) // don't bother with partial events
+		{
+			return false;
+		}
+
+		const uint32_t ExtraBytes = EventSize - EmptySize;
+		if (ExtraBytes > 0)
+		{
+			memcpy(((uint8_t*)Event) + EmptySize, Data + *Offset + EmptySize, ExtraBytes);
+		}
+
+		return true;
+	}
+
+	inline void Reset()
+	{
+		Head.store(0);
+		Tail.store(0);
+	}
+};
+
+static struct
+{
+	std::atomic_flag bSetup = {false};
+
+	std::atomic_flag bDisableEventsGlobal = {false};
+	// can't use thread_local because it will malloc during creation causing recursion
+	TLSFlag DisableEventsPerThread = {};
+
+	std::atomic<ScudoMemoryTraceHook> TraceHook = {nullptr};
+
+	std::atomic<uint32_t> EventCount = {0};
+
+	Spinlock WriteLock = {};
+	Spinlock SetupLock = {};
+
+	std::atomic<int> Fd = {0};
+	std::atomic<off64_t> FdOffset = {0};
+
+	EventBuffer Front, Back = {};
+	std::atomic<EventBuffer*> FrontBuffer = {nullptr};
+
+	std::atomic<uint64_t> StackTopMainThread = {0};
+
+#if !USE_SCUDO_HOOKS
+	void* (*DefaultMalloc)(size_t);
+	void* (*DefaultCalloc)(size_t, size_t);
+	void* (*DefaultRealloc)(void*, size_t);
+	void* (*DefaultMemAlign)(size_t, size_t);
+	void  (*DefaultFree)(void*);
+	int   (*DefaultPosixMemAlign)(void**, size_t, size_t);
+	void* (*DefaultAlignedAlloc)(size_t, size_t);
+	void* (*DefaultReallocArray)(void*, size_t, size_t);
+#endif
+} Ctx = {};
+
+struct GlobalReentranceGuard
+{
+	GlobalReentranceGuard()
+	{
+		Ctx.bDisableEventsGlobal.test_and_set();
+	}
+
+	~GlobalReentranceGuard()
+	{
+		Ctx.bDisableEventsGlobal.clear();
+	}
+
+private:
+	GlobalReentranceGuard(GlobalReentranceGuard const&) = delete;
+	GlobalReentranceGuard& operator=(GlobalReentranceGuard const&) = delete;
+};
+
+struct PerThreadTraceReentranceGuard
+{
+	PerThreadTraceReentranceGuard()
+	{
+		Ctx.DisableEventsPerThread.Set(true);
+	}
+
+	~PerThreadTraceReentranceGuard()
+	{
+		Ctx.DisableEventsPerThread.Set(false);
+	}
+
+private:
+	PerThreadTraceReentranceGuard(PerThreadTraceReentranceGuard const&) = delete;
+	PerThreadTraceReentranceGuard& operator=(PerThreadTraceReentranceGuard const&) = delete;
+};
+
+static inline void* LookupSymbol(const char* Name)
+{
+	void* Symbol = dlsym(RTLD_NEXT, Name);
+	if (!Symbol)
+	{
+		abort();
+	}
+	return Symbol;
+}
+
+static inline void PathWithPid(char* Buffer, size_t BufferSize, const char* Prefix)
+{
+	const size_t PrefixLen = strlen(Prefix); // bionic strlen() doesn't allocate
+	if (BufferSize < PrefixLen + 8)
+	{
+		return;
+	}
+
+	memcpy(Buffer, Prefix, PrefixLen);
+
+	int Pid = getpid(); // bionic getpid() doesn't allocate
+
+	// can't use snprintf due to internal allocations
+	Buffer[PrefixLen + 7] = '\0';
+	for (int i = 6; i >= 0; i--, Pid /= 10)
+	{
+		Buffer[PrefixLen + i] = '0' + (Pid % 10);
+	}
+}
+
+static inline void Setup()
+{
+	if (Ctx.bSetup.test(std::memory_order::relaxed))
+	{
+		return;
+	}
+
+	std::lock_guard Lock(Ctx.SetupLock);
+
+	if (Ctx.bSetup.test() == true)
+	{
+		return;
+	}
+
+	{
+		// Disable any tracing for duration of TLS creation
+		GlobalReentranceGuard Guard;
+		Ctx.DisableEventsPerThread.Setup();
+	}
+
+	if constexpr (bTraceEventsBeforeLibUnrealIsLoaded)
+	{
+		if constexpr (DebugFilePath == nullptr)
+		{
+			char* InternalFolderPath = getenv("UE_INTERNAL_FOLDER");
+			if (InternalFolderPath == nullptr)
+			{
+				abort();
+			}
+			// O_TMPFILE will be removed after it's closed
+			Ctx.Fd = open(InternalFolderPath, O_TMPFILE | O_EXCL | O_RDWR, S_IRUSR | S_IWUSR);
+		}
+		else
+		{
+			char Path[128];
+			PathWithPid(Path, sizeof(Path), DebugFilePath);
+			Ctx.Fd = open(Path, O_CREAT | O_TRUNC | O_RDWR, S_IRUSR | S_IWUSR | S_IRGRP | S_IWGRP);
+		}
+
+		// Disable events if backing storage is not created.
+		if (Ctx.Fd == 0)
+		{
+			Ctx.bDisableEventsGlobal.test_and_set();
+		}
+	}
+
+#if !USE_SCUDO_HOOKS
+	Ctx.DefaultMalloc        = (decltype(Ctx.DefaultMalloc))        LookupSymbol("malloc");
+	Ctx.DefaultCalloc        = (decltype(Ctx.DefaultCalloc))        LookupSymbol("calloc");
+	Ctx.DefaultRealloc       = (decltype(Ctx.DefaultRealloc))       LookupSymbol("realloc");
+	Ctx.DefaultMemAlign      = (decltype(Ctx.DefaultMemAlign))      LookupSymbol("memalign");
+	Ctx.DefaultFree          = (decltype(Ctx.DefaultFree))          LookupSymbol("free");
+	Ctx.DefaultPosixMemAlign = (decltype(Ctx.DefaultPosixMemAlign)) LookupSymbol("posix_memalign");
+	Ctx.DefaultAlignedAlloc  = (decltype(Ctx.DefaultAlignedAlloc))  LookupSymbol("aligned_alloc");
+	Ctx.DefaultReallocArray  = (decltype(Ctx.DefaultReallocArray))  LookupSymbol("reallocarray");
+#endif
+
+	Ctx.FrontBuffer = &Ctx.Front;
+
+	Ctx.bSetup.test_and_set();
+}
+
+static inline void CloseTempFileAndDisableTracing()
+{
+	Ctx.bDisableEventsGlobal.test_and_set();
+	const int Fd = Ctx.Fd.exchange(0);
+	if (Fd != 0)
+	{
+		close(Fd);
+	}
+}
+
+static inline void ReportEventToHook(const ScudoMemoryTraceHook TraceHook, const FEvent* Event)
+{
+	TraceHook(Event->Type, Event->Ptr, Event->Size, Event->Alignment, Event->Frames, Event->FrameCount);
+}
+
+static inline void ReportEvent(const FEvent* Event)
+{
+	ScudoMemoryTraceHook TraceHook = Ctx.TraceHook.load(std::memory_order_relaxed);
+
+	if (TraceHook != nullptr)
+	{
+		PerThreadTraceReentranceGuard Guard;
+		ReportEventToHook(TraceHook, Event);
+	}
+	else if constexpr(bTraceEventsBeforeLibUnrealIsLoaded)
+	{
+		const uint32_t EventCount = Ctx.EventCount.fetch_add(1);
+
+		// Disable tracer if UE never set the hook
+		if (Ctx.Fd.load(std::memory_order_relaxed) != 0 &&
+			EventCount > MaxEventsBeforeHookIsSet)
+		{
+			CloseTempFileAndDisableTracing();
+			return;
+		}
+
+		while (Ctx.FrontBuffer.load()->Write(Event) == false)
+		{
+			std::lock_guard Lock(Ctx.WriteLock);
+
+			// Report again if trace hook was installed right when we were waiting for the lock.
+			if (Ctx.TraceHook.load() != nullptr)
+			{
+				ReportEvent(Event);
+				return;
+			}
+
+			EventBuffer* FrontBuffer = Ctx.FrontBuffer.load();
+
+			// Try to write to front buffer again, maybe it was swapped while we were waiting for the lock
+			if (FrontBuffer->Write(Event))
+			{
+				return;
+			}
+
+			// If writing fails again, check if we got lock out of events while waiting for the lock, maybe Fd was closed.
+			if (Ctx.bDisableEventsGlobal.test())
+			{
+				return;
+			}
+
+			// Current buffer is full, swap buffers to unlock new writes on other threads
+			EventBuffer* BackBuffer = FrontBuffer;
+			FrontBuffer = BackBuffer == &Ctx.Front ? &Ctx.Back : &Ctx.Front;
+			Ctx.FrontBuffer = FrontBuffer;
+
+			// Try to write to new buffer, so our event will be closer to the front of the buffer
+			bool bReturnAfterFlush = false;
+			if (FrontBuffer->Write(Event))
+			{
+				bReturnAfterFlush = true;
+			}
+			else
+			{
+				// New front buffer is already full? Will try again in the next loop.
+			}
+
+			const uint32_t BackBufferSize = BackBuffer->Head.load();
+			const int Fd = Ctx.Fd.load();
+			if (Fd == 0)
+			{
+				Ctx.bDisableEventsGlobal.test_and_set();
+				return;
+			}
+			const off64_t Cursor = Ctx.FdOffset.fetch_add(BackBufferSize);
+			ftruncate64(Fd, Cursor + BackBufferSize);
+
+			uint32_t BytesWritten = 0;
+			while (BytesWritten < BackBufferSize)
+			{
+				const ssize_t Result = pwrite64(Fd, BackBuffer->Data + BytesWritten, BackBufferSize - BytesWritten, Cursor + BytesWritten);
+				if (Result > 0)
+				{
+					BytesWritten += Result;
+				}
+				else if (Result == -1) // Failed to write to backing storage, disable tracing
+				{
+					CloseTempFileAndDisableTracing();
+					return;
+				}
+			}
+
+			BackBuffer->Reset();
+
+			if (bReturnAfterFlush)
+			{
+				return;
+			}
+		}
+
+		// Periodically report time events, only needed before the hook is set, configuration is the same to MemoryTrace.cpp
+		if constexpr(bGenerateTimeEvents)
+		{
+			if ((EventCount & BeforeHookMarkerSamplePeriod) == 0)
+			{
+				timespec ts;
+				clock_gettime(CLOCK_MONOTONIC, &ts);
+				const uint64_t Cycles = ((((uint64_t)ts.tv_sec) * 1000000ULL) + (((uint64_t)ts.tv_nsec) / 1000ULL));
+
+				FEvent TimeEvent = {
+					.Ptr = Cycles,
+					.Type = EScudoEventType::Time
+				};
+
+				// Beware, recursion
+				ReportEvent(&TimeEvent);
+			}
+		}
+	}
+}
+
+static inline bool ReadEventFromStorage(const int Fd, FEvent* Event, off64_t* ReadPos)
+{
+	uint8_t* Buffer = (uint8_t*)Event;
+	uint32_t BytesRead = 0;
+	for (uint32_t i = 0; i < 2; ++i)
+	{
+		uint32_t BytesToRead = i == 0 ? FEvent::GetEmptyEventSize() : Event->GetEventSize();
+
+		while (BytesRead < BytesToRead)
+		{
+			const ssize_t Result = pread64(Fd, Buffer + BytesRead, BytesToRead - BytesRead, *ReadPos);
+			if (Result < 0)
+			{
+				return false;
+			}
+			BytesRead += Result;
+			*ReadPos += Result;
+		}
+	}
+
+	return true;
+}
+
+// ---------------------------------------------------------------------------------------------------------------------
+
+static inline uint64_t RemovePtrTags(const void* RawPtr)
+{
+	uint64_t Raw = (uint64_t)RawPtr;
+	// Remove bionic pointer tag. For more info see https://cs.android.com/android/platform/superproject/main/+/main:bionic/libc/bionic/malloc_tagged_pointers.h
+	Raw &= 0x00FF'FFFF'FFFF'FFFFull;
+#if defined(__aarch64__)
+	// Remove PAC
+	asm("xpaclri" : "+r" (Raw));
+	return Raw;
+#else
+	return Raw;
+#endif
+}
+
+static inline void CollectCallstack(const void* FrameStartPtr, const void* ReturnAddrPtr, FEvent* OutEvent)
+{
+	struct StackFrame
+	{
+		StackFrame* NextFrame;
+		void* ReturnPtr;
+	};
+
+	// There is no easy way to get current thread stack top without invoking malloc somewhere in the chain of calls, hence assume some sensible thread size.
+	StackFrame* FrameEnd = (StackFrame*)((uint8_t*)FrameStartPtr + 16 * PAGE_SIZE); // 64kb
+	StackFrame* FrameStart = (StackFrame*)FrameStartPtr;
+
+	for (StackFrame* CurrentFrame = FrameStart;
+		OutEvent->FrameCount < FEvent::MaxCallstackFrameCount &&
+		CurrentFrame->NextFrame > FrameStart &&
+		CurrentFrame->NextFrame <= FrameEnd &&
+		((uintptr_t)CurrentFrame->NextFrame & (sizeof(StackFrame) - 1)) == 0; // stop at unaligned frame
+		CurrentFrame = CurrentFrame->NextFrame)
+	{
+		const uint64_t Ptr = RemovePtrTags(CurrentFrame->ReturnPtr);
+		if (Ptr != 0)
+		{
+			OutEvent->Frames[OutEvent->FrameCount++] = Ptr;
+		}
+		else
+		{
+			break;
+		}
+	}
+
+	// If any unwinding fails, record return pointer to get at least some info.
+	if (OutEvent->FrameCount == 0)
+	{
+		OutEvent->Frames[OutEvent->FrameCount++] = RemovePtrTags(ReturnAddrPtr);
+	}
+}
+// ---------------------------------------------------------------------------------------------------------------------
+
+static inline bool ShouldIgnoreEvent()
+{
+	return Ctx.bDisableEventsGlobal.test() || Ctx.DisableEventsPerThread.Get();
+}
+
+// ---------------------------------------------------------------------------------------------------------------------
+
+static inline void ReportAllocEventImpl(const void* FrameStartPtr, const void* ReturnAddrPtr, const void* RawPtr, const size_t Size, const size_t Alignment, const uint8_t Hint)
+{
+	if (ShouldIgnoreEvent())
+	{
+		return;
+	}
+
+	const uint64_t Ptr = RemovePtrTags(RawPtr);
+	if (Ptr == 0 || Size == 0)
+	{
+		return;
+	}
+
+	FEvent Event = {
+		.Ptr = Ptr,
+		.Size = Size,
+		.Alignment = (uint32_t)Alignment,
+		.Type = EScudoEventType::Alloc,
+	};
+
+	CollectCallstack(FrameStartPtr, ReturnAddrPtr, &Event);
+	ReportEvent(&Event);
+}
+
+#define ReportAllocEvent(Ptr, Size, Alignment, Hint) ReportAllocEventImpl(__builtin_frame_address(0), __builtin_return_address(0), Ptr, Size, Alignment, Hint)
+
+static inline void ReportFreeEventImpl(const void* FrameStartPtr, const void* ReturnAddrPtr, const void* RawPtr, const uint8_t Hint)
+{
+	if (ShouldIgnoreEvent())
+	{
+		return;
+	}
+
+	const uint64_t Ptr = RemovePtrTags(RawPtr);
+	if (Ptr == 0)
+	{
+		return;
+	}
+
+	FEvent Event = {
+		.Ptr = Ptr,
+		.Type = EScudoEventType::Free,
+	};
+
+	CollectCallstack(FrameStartPtr, ReturnAddrPtr, &Event);
+	ReportEvent(&Event);
+}
+
+#define ReportFreeEvent(Ptr, Hint) ReportFreeEventImpl(__builtin_frame_address(0), __builtin_return_address(0), Ptr, Hint)
+
+// ---------------------------------------------------------------------------------------------------------------------
+
+void ScudoMemoryTrace_SetHook(ScudoMemoryTraceHook TraceHook)
+{
+	Ctx.WriteLock.lock();
+
+	Ctx.TraceHook.store(TraceHook);
+
+	// Start reporting again even if we disabled reporting events previously.
+	Ctx.bDisableEventsGlobal.clear();
+
+	const int Fd = Ctx.Fd.exchange(0);
+	const off64_t FileSize = Ctx.FdOffset.exchange(0);
+
+	Ctx.WriteLock.unlock();
+
+	if (Ctx.Fd == 0)
+	{
+		return;
+	}
+
+	if constexpr(bTraceEventsBeforeLibUnrealIsLoaded)
+	{
+		// At this point all new events will be flushed to our hook. Proceed with draining backing storage and event buffers.
+		lseek(Ctx.Fd, 0, SEEK_SET);
+
+		PerThreadTraceReentranceGuard Guard;
+
+		FEvent Event;
+
+		off64_t ReadPos = 0;
+		while(ReadPos < FileSize && ReadEventFromStorage(Fd, &Event, &ReadPos))
+		{
+			ReportEventToHook(TraceHook, &Event);
+		}
+
+		close(Fd);
+
+		uint32_t EventBufferOffset = 0;
+		uint32_t EventBufferDataSize = Ctx.Back.Head.load();
+		while(Ctx.Back.Read(&Event, &EventBufferOffset, EventBufferDataSize))
+		{
+			ReportEventToHook(TraceHook, &Event);
+		}
+
+		Ctx.Back.Reset();
+
+		EventBufferOffset = 0;
+		EventBufferDataSize = Ctx.Front.Head.load();
+		while(Ctx.Front.Read(&Event, &EventBufferOffset, EventBufferDataSize))
+		{
+			ReportEventToHook(TraceHook, &Event);
+		}
+
+		Ctx.Front.Reset();
+	}
+}
+
+// ---------------------------------------------------------------------------------------------------------------------
+
+#if USE_SCUDO_HOOKS
+
+extern "C" void __scudo_allocate_hook(void* Ptr, size_t Size)
+{
+	ReportAllocEvent(Ptr, Size, 0);
+}
+
+extern "C" void __scudo_deallocate_hook(void* Ptr)
+{
+	ReportFreeEvent(Ptr, 1);
+}
+
+#else
+
+extern "C" void* malloc(size_t Size)
+{
+	Setup();
+
+	void* Ptr = Ctx.DefaultMalloc(Size);
+	ReportAllocEvent(Ptr, Size, 0, 0);
+	return Ptr;
+}
+
+extern "C" void* calloc(size_t Num, size_t Size)
+{
+	Setup();
+
+	void* Ptr = Ctx.DefaultCalloc(Num, Size);
+	ReportAllocEvent(Ptr, Size, 0, 1);
+	return Ptr;
+}
+
+extern "C" void* realloc(void* OldPtr, size_t NewSize)
+{
+	Setup();
+
+	ReportFreeEvent(OldPtr, 2);
+	void* NewPtr = Ctx.DefaultRealloc(OldPtr, NewSize);
+	ReportAllocEvent(NewPtr, NewSize, 0, 2);
+	return NewPtr;
+}
+
+extern "C" void* memalign(size_t Alignment, size_t Size)
+{
+	Setup();
+
+	void* Ptr = Ctx.DefaultMemAlign(Alignment, Size);
+	ReportAllocEvent(Ptr, Size, Alignment, 3);
+	return Ptr;
+}
+
+extern "C" void free(void* Ptr)
+{
+	Setup();
+
+	ReportFreeEvent(Ptr, 4);
+	Ctx.DefaultFree(Ptr);
+}
+
+extern "C" int posix_memalign(void** Ptr, size_t Alignment, size_t Size)
+{
+	Setup();
+
+	int Result = Ctx.DefaultPosixMemAlign(Ptr, Alignment, Size);
+	ReportAllocEvent((Ptr ? *Ptr : nullptr), Size, Alignment, 5);
+	return Result;
+}
+
+extern "C" void* aligned_alloc(size_t Alignment, size_t Size)
+{
+	Setup();
+
+	void* Ptr = Ctx.DefaultAlignedAlloc(Alignment, Size);
+	ReportAllocEvent(Ptr, Size, Alignment, 6);
+	return Ptr;
+}
+
+extern "C" void* reallocarray(void* OldPtr, size_t ItemCount, size_t ItemSize)
+{
+	Setup();
+
+	ReportFreeEvent(OldPtr, 7);
+	void* NewPtr = Ctx.DefaultReallocArray(OldPtr, ItemCount, ItemSize);
+	ReportAllocEvent(NewPtr, ItemCount * ItemSize, 0, 7);
+	return NewPtr;
+}
+
+#endif
diff --git a/Engine/Build/Android/Prebuilt/ScudoMemoryTrace/ScudoMemoryTrace.h b/Engine/Build/Android/Prebuilt/ScudoMemoryTrace/ScudoMemoryTrace.h
new file mode 100644
--- /dev/null
+++ b/Engine/Build/Android/Prebuilt/ScudoMemoryTrace/ScudoMemoryTrace.h
@@ -0,0 +1,17 @@
+// Copyright Epic Games, Inc. All Rights Reserved.
+
+#pragma once
+
+#include <stdint.h>
+
+enum class EScudoEventType : uint8_t
+{
+	Invalid = 0,
+	Alloc = 1, // Ptr != 0, Size != 0, Alignment might be 0 if not specified
+	Free = 2,  // Ptr != 0, Size == 0, Alignment == 0
+	Time = 3   // Time in stored in Ptr, Size == 0, Alignment == 0
+};
+
+typedef void (*ScudoMemoryTraceHook)(EScudoEventType EventType, uint64_t Ptr, uint64_t Size, uint32_t Alignment, const void* Frames, uint32_t FrameCount);
+
+extern "C" void ScudoMemoryTrace_SetHook(ScudoMemoryTraceHook TraceHook);
diff --git a/Engine/Build/Android/Prebuilt/ScudoMemoryTrace/wrap.sh b/Engine/Build/Android/Prebuilt/ScudoMemoryTrace/wrap.sh
new file mode 100644
--- /dev/null
+++ b/Engine/Build/Android/Prebuilt/ScudoMemoryTrace/wrap.sh
@@ -0,0 +1,5 @@
+#!/system/bin/sh
+HERE="$(cd "$(dirname "$0")" && pwd)"
+export UE_INTERNAL_FOLDER="/data/user/0/${PACKAGE_NAME}/files/"
+export LD_PRELOAD="$HERE/libScudoMemoryTrace.so"
+"$@"
diff --git a/Engine/Source/Programs/UnrealBuildTool/Platform/Android/UEDeployAndroid.cs b/Engine/Source/Programs/UnrealBuildTool/Platform/Android/UEDeployAndroid.cs
--- a/Engine/Source/Programs/UnrealBuildTool/Platform/Android/UEDeployAndroid.cs
+++ b/Engine/Source/Programs/UnrealBuildTool/Platform/Android/UEDeployAndroid.cs
@@ -87,6 +87,12 @@ namespace UnrealBuildTool
 		[CommandLine("-ForceDontBundleLibrariesInAPK=")]
 		public bool? ForceDontBundleLibrariesInAPK;
 
+		/// <summary>
+		/// Adds LD_PRELOAD'ed .so to capture all malloc/free/etc calls to libc.so and route them to our memory tracing.
+		/// </summary>
+		[CommandLine("-ScudoMemoryTracing", Value = "true")]
+		public bool bEnableScudoMemoryTracing = false;
+
 		public UEDeployAndroid(FileReference? InProjectFile, bool InForcePackageData, ILogger InLogger)
 			: base(InLogger)
 		{
@@ -1558,6 +1564,57 @@ namespace UnrealBuildTool
 			}
 		}
 
+		void CopyScudoMemoryTracerLib(string UnrealBuildPath, UnrealArch UnrealArch, string NDKArch, string PackageName)
+		{
+			if (UnrealArch != UnrealArch.Arm64 && UnrealArch != UnrealArch.X64)
+			{
+				throw new BuildException("ScudoMemoryTrace not supported for arch {0}", UnrealArch.ToString());
+			}
+
+			string ScudoMemoryTraceLib = Path.Combine(Unreal.EngineDirectory.ToString(), "Build", "Android", "Prebuilt", "ScudoMemoryTrace", NDKArch, "libScudoMemoryTrace.so");
+
+			string WrapSh = Path.Combine(Unreal.EngineDirectory.ToString(), "Build", "Android", "Prebuilt", "ScudoMemoryTrace", "wrap.sh");
+
+			if (File.Exists(ScudoMemoryTraceLib) && File.Exists(WrapSh))
+			{
+				string LibDestDir = Path.Combine(UnrealBuildPath, "libs", NDKArch);
+				string LibDest = Path.Combine(LibDestDir, "libScudoMemoryTrace.so");
+				Directory.CreateDirectory(LibDestDir);
+				if (!CopyIfDifferent(ScudoMemoryTraceLib, LibDest, true, false))
+				{
+					Logger.LogInformation("{LibDest} is up to date", LibDest);
+				}
+
+				string WrapDestDir = Path.Combine(UnrealBuildPath, "resources", "lib", NDKArch);
+				Directory.CreateDirectory(WrapDestDir);
+				string WrapDestFilePath = Path.Combine(WrapDestDir, "wrap.sh");
+
+				Dictionary<string, string> Replacements = new Dictionary<string, string>
+				{
+					{ "${PACKAGE_NAME}", PackageName },
+				};
+				string WrapShText = File.ReadAllText(WrapSh);
+				foreach (KeyValuePair<string, string> KVP in Replacements)
+				{
+					WrapShText = WrapShText.Replace(KVP.Key, KVP.Value);
+				}
+
+				if (!File.Exists(WrapDestFilePath) || File.ReadAllText(WrapDestFilePath) != WrapShText)
+				{
+					Logger.LogInformation("Writing {WrapDestFilePath}", WrapDestFilePath);
+					File.WriteAllText(WrapDestFilePath, WrapShText);
+				}
+				else
+				{
+					Logger.LogInformation("{WrapDestFilePath} is up to date", WrapDestFilePath);
+				}
+			}
+			else
+			{
+				throw new BuildException("No ScudoMemoryTrace found in {0} or wrap.sh in {1}", ScudoMemoryTraceLib, WrapSh);
+			}
+		}
+
 		private static int RunCommandLineProgramAndReturnResult(string WorkingDirectory, string Command, string Params, ILogger Logger, string? OverrideDesc = null, bool bUseShellExecute = false)
 		{
 			// Process Arguments follow windows conventions in .NET Core
@@ -2720,7 +2777,7 @@ namespace UnrealBuildTool
 			Text.AppendLine("\t             android:icon=\"@drawable/icon\"");
 
 			AndroidToolChain.ClangSanitizer Sanitizer = ToolChain.BuildWithSanitizer();
-			if ((Sanitizer != AndroidToolChain.ClangSanitizer.None && Sanitizer != AndroidToolChain.ClangSanitizer.HwAddress))
+			if ((Sanitizer != AndroidToolChain.ClangSanitizer.None && Sanitizer != AndroidToolChain.ClangSanitizer.HwAddress) || bEnableScudoMemoryTracing)
 			{
 				bExtractNativeLibs = true;
 			}
@@ -3847,7 +3904,7 @@ namespace UnrealBuildTool
 			bool bExtractNativeLibs = true;
 			Ini.GetBool("/Script/AndroidRuntimeSettings.AndroidRuntimeSettings", "bExtractNativeLibs", out bExtractNativeLibs);
 			AndroidToolChain.ClangSanitizer Sanitizer = ToolChain.BuildWithSanitizer();
-			if ((Sanitizer != AndroidToolChain.ClangSanitizer.None && Sanitizer != AndroidToolChain.ClangSanitizer.HwAddress))
+			if ((Sanitizer != AndroidToolChain.ClangSanitizer.None && Sanitizer != AndroidToolChain.ClangSanitizer.HwAddress) || bEnableScudoMemoryTracing)
 			{
 				bExtractNativeLibs = true;
 			}
@@ -4152,6 +4209,10 @@ namespace UnrealBuildTool
 			{
 				BuildTypeID += Sanitizer.ToString() + "Sanitizer";
 			}
+			else if (bEnableScudoMemoryTracing)
+			{
+				BuildTypeID += Sanitizer.ToString() + "ScudoMemoryTrace";
+			}
 			if (File.Exists(BuildTypeFilename))
 			{
 				string BuildTypeContents = File.ReadAllText(BuildTypeFilename);
@@ -4961,6 +5022,15 @@ popd
 				if (Sanitizer != AndroidToolChain.ClangSanitizer.None)
 				{
 					CopyClangSanitizerLib(UnrealBuildPath, Arch, NDKArch, Sanitizer);
+
+					if (bEnableScudoMemoryTracing)
+					{
+						throw new BuildException("ScudoMemoryTrace not supported with sanitizer enabled");
+					}
+				}
+				else if (bEnableScudoMemoryTracing)
+				{
+					CopyScudoMemoryTracerLib(UnrealBuildPath, Arch, NDKArch, PackageName);
 				}
 
 				// copy postbuild plugin files
@@ -5023,7 +5093,7 @@ popd
 
 				CleanCopyDirectory(Path.Combine(UnrealBuildPath, "jni"), Path.Combine(UnrealBuildGradleMainPath, "jniLibs"), Excludes);  // has debug symbols
 				CleanCopyDirectory(Path.Combine(UnrealBuildPath, "libs"), Path.Combine(UnrealBuildGradleMainPath, "libs"), Excludes);
-				if (Sanitizer != AndroidToolChain.ClangSanitizer.None && Sanitizer != AndroidToolChain.ClangSanitizer.HwAddress)
+				if ((Sanitizer != AndroidToolChain.ClangSanitizer.None && Sanitizer != AndroidToolChain.ClangSanitizer.HwAddress) || bEnableScudoMemoryTracing)
 				{
 					CleanCopyDirectory(Path.Combine(UnrealBuildPath, "resources"), Path.Combine(UnrealBuildGradleMainPath, "resources"), Excludes);
 				}
diff --git a/Engine/Source/Runtime/Core/Core.Build.cs b/Engine/Source/Runtime/Core/Core.Build.cs
--- a/Engine/Source/Runtime/Core/Core.Build.cs
+++ b/Engine/Source/Runtime/Core/Core.Build.cs
@@ -153,7 +153,11 @@ public class Core : ModuleRules
 				PublicDefinitions.Add("UE_MEMORY_TRACE_AVAILABLE=1");
 				PublicDefinitions.Add("UE_MEMORY_TAGS_TRACE_ENABLED=1");
 				PublicDefinitions.Add("UE_CALLSTACK_TRACE_ENABLED=1");
-				PublicDefinitions.Add("UE_CALLSTACK_TRACE_ANDROID_USE_STACK_FRAMES_WALKING=1");
+				PrivateDefinitions.Add("UE_CALLSTACK_TRACE_ANDROID_USE_STACK_FRAMES_WALKING=1");
+
+				// Support for memory tracing libc.so malloc
+				PrivateIncludePaths.Add(Path.Combine(EngineDirectory, "Build", "Android", "Prebuilt", "ScudoMemoryTrace"));
+				PrivateDefinitions.Add("UE_MEMORY_TRACE_ANDROID_ENABLE_SCUDO_TRACING_SUPPORT=1");
 			}
 		}
 		else if (Target.IsInPlatformGroup(UnrealPlatformGroup.Unix))
diff --git a/Engine/Source/Runtime/Core/Private/Android/AndroidCallstackTrace.cpp b/Engine/Source/Runtime/Core/Private/Android/AndroidCallstackTrace.cpp
--- a/Engine/Source/Runtime/Core/Private/Android/AndroidCallstackTrace.cpp
+++ b/Engine/Source/Runtime/Core/Private/Android/AndroidCallstackTrace.cpp
@@ -1,5 +1,6 @@
 ﻿// Copyright Epic Games, Inc. All Rights Reserved.
 
+#include "AndroidCallstackTrace.h"
 #include "ProfilingDebugging/CallstackTrace.h"
 
 #if UE_CALLSTACK_TRACE_ENABLED
@@ -20,6 +21,7 @@ public:
 	FBacktracer(FMalloc* InMalloc);
 	~FBacktracer();
 	static FBacktracer*	Get();
+	FORCEINLINE uint32 GetBacktraceId(uint64* StackFrames, uint32 NumStackFrames);
 	FORCEINLINE uint32 GetBacktraceId(void* ReturnAddress);
 
 private:
@@ -51,7 +53,37 @@ FBacktracer* FBacktracer::Get()
 }
 
 ////////////////////////////////////////////////////////////////////////////////
+FORCEINLINE uint32 FBacktracer::GetBacktraceId(uint64* StackFrames, uint32 NumStackFrames)
+{
+#if !UE_BUILD_SHIPPING
+	if (NumStackFrames == 0)
+	{
+		return 0;
+	}
+
+	FCallstackTracer::FBacktraceEntry BacktraceEntry;
+	uint64 BacktraceId = 0;
+	for (int32 Index = 0; Index < NumStackFrames; Index++)
+	{
+		// This is a simple order-dependent LCG. Should be sufficient enough
+		BacktraceId += StackFrames[Index];
+		BacktraceId *= 0x30be8efa499c249dull;
+	}
 
+	// Save the collected id
+	BacktraceEntry.Hash = BacktraceId;
+	BacktraceEntry.FrameCount = NumStackFrames;
+	BacktraceEntry.Frames = StackFrames;
+
+	// Add to queue to be processed. This might block until there is room in the
+	// queue (i.e. the processing thread has caught up processing).
+	return CallstackTracer.AddCallstack(BacktraceEntry);
+#endif
+
+	return 0;
+}
+
+////////////////////////////////////////////////////////////////////////////////
 FORCEINLINE uint32 FBacktracer::GetBacktraceId(void* ReturnAddress)
 {
 #if !UE_BUILD_SHIPPING
@@ -65,39 +97,18 @@ FORCEINLINE uint32 FBacktracer::GetBacktraceId(void* ReturnAddress)
 
 	if (NumStackFrames > 0)
 	{
-		FCallstackTracer::FBacktraceEntry BacktraceEntry;
-		uint64 BacktraceId = 0;
-		uint32 FrameIdx = 0;
-		bool bUseAddress = false;
+		// Skip until we find ReturnAddress in the stack, otherwise start from the beginning.
+		uint32 StartFromIndex = 0;
 		for (int32 Index = 0; Index < NumStackFrames; Index++)
 		{
-			if (!bUseAddress)
+			if (StackFrames[Index] == (uint64)ReturnAddress)
 			{
-				// start using backtrace only after ReturnAddress
-				if (StackFrames[Index] == (uint64)ReturnAddress)
-				{
-					bUseAddress = true;
-				}
-			}
-			if (bUseAddress || NumStackFrames == 1)
-			{
-				uint64 RetAddr = StackFrames[Index];
-				StackFrames[FrameIdx++] = RetAddr;
-
-				// This is a simple order-dependent LCG. Should be sufficient enough
-				BacktraceId += RetAddr;
-				BacktraceId *= 0x30be8efa499c249dull;
+				StartFromIndex = Index;
+				break;
 			}
 		}
 
-		// Save the collected id
-		BacktraceEntry.Hash = BacktraceId;
-		BacktraceEntry.FrameCount = FrameIdx;
-		BacktraceEntry.Frames = StackFrames;
-
-		// Add to queue to be processed. This might block until there is room in the
-		// queue (i.e. the processing thread has caught up processing).
-		return CallstackTracer.AddCallstack(BacktraceEntry);
+		return GetBacktraceId(StackFrames + StartFromIndex, NumStackFrames - StartFromIndex);
 	}
 #endif
 
@@ -143,4 +154,15 @@ uint32 CallstackTrace_GetCurrentId()
 	return 0;
 }
 
+////////////////////////////////////////////////////////////////////////////////
+uint32 CallstackTrace_GetExternalCallstackId(uint64* Frames, uint32 Count)
+{
+	if (FBacktracer* Instance = FBacktracer::Get())
+	{
+		return Instance->GetBacktraceId(Frames, Count);
+	}
+
+	return 0;
+}
+
 #endif
\ No newline at end of file
diff --git a/Engine/Source/Runtime/Core/Private/Android/AndroidCallstackTrace.h b/Engine/Source/Runtime/Core/Private/Android/AndroidCallstackTrace.h
new file mode 100644
--- /dev/null
+++ b/Engine/Source/Runtime/Core/Private/Android/AndroidCallstackTrace.h
@@ -0,0 +1,13 @@
+﻿// Copyright Epic Games, Inc. All Rights Reserved.
+
+#pragma once
+
+#include "ProfilingDebugging/CallstackTrace.h"
+
+#if UE_CALLSTACK_TRACE_ENABLED
+
+#include "CoreTypes.h"
+
+uint32 CallstackTrace_GetExternalCallstackId(uint64* Frames, uint32 Count);
+
+#endif
\ No newline at end of file
diff --git a/Engine/Source/Runtime/Core/Private/Android/AndroidPlatformMisc.cpp b/Engine/Source/Runtime/Core/Private/Android/AndroidPlatformMisc.cpp
--- a/Engine/Source/Runtime/Core/Private/Android/AndroidPlatformMisc.cpp
+++ b/Engine/Source/Runtime/Core/Private/Android/AndroidPlatformMisc.cpp
@@ -48,6 +48,7 @@
 
 #include "Android/AndroidPlatformStackWalk.h"
 #include "Android/AndroidSignals.h"
+#include "AndroidScudoMemoryTrace.h"
 
 #include "Misc/OutputDevice.h"
 #include "Logging/LogMacros.h"
@@ -589,8 +590,9 @@ void FAndroidMisc::PlatformInit()
 	AndroidThunkJava_AddNetworkListener();
 #endif
 
-
 	InitCpuThermalSensor();
+
+	AndroidScudoMemoryTrace::Init();
 }
 
 extern void AndroidThunkCpp_DismissSplashScreen();
diff --git a/Engine/Source/Runtime/Core/Private/Android/AndroidScudoMemoryTrace.cpp b/Engine/Source/Runtime/Core/Private/Android/AndroidScudoMemoryTrace.cpp
new file mode 100644
--- /dev/null
+++ b/Engine/Source/Runtime/Core/Private/Android/AndroidScudoMemoryTrace.cpp
@@ -0,0 +1,70 @@
+// Copyright Epic Games, Inc. All Rights Reserved.
+
+#include "AndroidScudoMemoryTrace.h"
+
+#ifndef UE_MEMORY_TRACE_ANDROID_ENABLE_SCUDO_TRACING_SUPPORT
+#	define UE_MEMORY_TRACE_ANDROID_ENABLE_SCUDO_TRACING_SUPPORT 0
+#endif
+
+#if 1 || UE_MEMORY_TRACE_ANDROID_ENABLE_SCUDO_TRACING_SUPPORT
+
+#include "AndroidCallstackTrace.h"
+#include "ProfilingDebugging/MemoryTrace.h"
+#include "ScudoMemoryTrace.h"
+
+#include <dlfcn.h>
+
+extern bool MemoryTrace_IsActive();
+
+static void ScudoMemoryTraceHookFn(EScudoEventType EventType, uint64_t Ptr, uint64_t Size, uint32_t Alignment, const void* Frames, uint32_t FrameCount)
+{
+	// LLM(UE_TRACE_METADATA_CLEAR_SCOPE());
+	// LLM(UE_MEMSCOPE(ELLMTag::)); // TODO create new LLM tag
+
+	const uint32 ExternalCallstackId = Frames ? CallstackTrace_GetExternalCallstackId((uint64*)Frames, FrameCount) : 0; 
+
+	switch (EventType)
+	{
+	case EScudoEventType::Alloc:
+	{
+		MemoryTrace_Alloc(Ptr, Size, Alignment, EMemoryTraceRootHeap::SystemMemory, ExternalCallstackId);
+		break;
+	}
+	case EScudoEventType::Free:
+	{
+		MemoryTrace_Free(Ptr, EMemoryTraceRootHeap::SystemMemory, ExternalCallstackId);
+		break;
+	}
+	case EScudoEventType::Time:
+	{
+		// Time event before libUnreal.so, ignore
+		break;
+	}
+	default:
+	{
+		break;
+	}
+	}
+}
+
+void AndroidScudoMemoryTrace::Init()
+{
+	if (!MemoryTrace_IsActive())
+	{
+		return;
+	}
+
+	// Find the set hook function dynamically, because libScudoMemoryTrace.so is LD_PRELOAD'ed
+	decltype(&ScudoMemoryTrace_SetHook) SetHookFn = (decltype(&ScudoMemoryTrace_SetHook))dlsym(RTLD_DEFAULT, "ScudoMemoryTrace_SetHook");
+
+	if (SetHookFn != nullptr)
+	{
+		SetHookFn(ScudoMemoryTraceHookFn);
+	}
+};
+
+#else
+
+void AndroidScudoMemoryTrace::Init() {}
+
+#endif
\ No newline at end of file
diff --git a/Engine/Source/Runtime/Core/Private/Android/AndroidScudoMemoryTrace.h b/Engine/Source/Runtime/Core/Private/Android/AndroidScudoMemoryTrace.h
new file mode 100644
--- /dev/null
+++ b/Engine/Source/Runtime/Core/Private/Android/AndroidScudoMemoryTrace.h
@@ -0,0 +1,8 @@
+﻿// Copyright Epic Games, Inc. All Rights Reserved.
+
+#pragma once
+
+struct AndroidScudoMemoryTrace
+{
+	static void Init();
+};
\ No newline at end of file
diff --git a/Engine/Source/Runtime/Core/Private/ProfilingDebugging/MemoryTrace.cpp b/Engine/Source/Runtime/Core/Private/ProfilingDebugging/MemoryTrace.cpp
--- a/Engine/Source/Runtime/Core/Private/ProfilingDebugging/MemoryTrace.cpp
+++ b/Engine/Source/Runtime/Core/Private/ProfilingDebugging/MemoryTrace.cpp
@@ -430,7 +430,7 @@ void MemoryTrace_UpdateInternal()
 }
 
 ////////////////////////////////////////////////////////////////////////////////
-void MemoryTrace_Alloc(uint64 Address, uint64 Size, uint32 Alignment, HeapId RootHeap)
+void MemoryTrace_Alloc(uint64 Address, uint64 Size, uint32 Alignment, HeapId RootHeap, uint32 ExternalCallstackId)
 {
 	if (!GTraceAllowed)
 	{
@@ -441,7 +441,7 @@ void MemoryTrace_Alloc(uint64 Address, uint64 Size, uint32 Alignment, HeapId Roo
 
 	const uint32 AlignmentPow2 = uint32(FPlatformMath::CountTrailingZeros(Alignment));
 	const uint32 Alignment_SizeLower = (AlignmentPow2 << SizeShift) | uint32(Size & ((1 << SizeShift) - 1));
-	const uint32 CallstackId = GDoNotAllocateInTrace ? 0 : CallstackTrace_GetCurrentId();
+	const uint32 CallstackId = ExternalCallstackId ? ExternalCallstackId : GDoNotAllocateInTrace ? 0 : CallstackTrace_GetCurrentId();
 
 	switch (RootHeap)
 	{
@@ -481,7 +481,7 @@ void MemoryTrace_Alloc(uint64 Address, uint64 Size, uint32 Alignment, HeapId Roo
 }
 
 ////////////////////////////////////////////////////////////////////////////////
-void MemoryTrace_Free(uint64 Address, HeapId RootHeap)
+void MemoryTrace_Free(uint64 Address, HeapId RootHeap, uint32 ExternalCallstackId)
 {
 	if (!GTraceAllowed)
 	{
@@ -490,7 +490,7 @@ void MemoryTrace_Free(uint64 Address, HeapId RootHeap)
 
 	checkSlow(RootHeap < 16);
 
-	const uint32 CallstackId = GDoNotAllocateInTrace ? 0 : CallstackTrace_GetCurrentId();
+	const uint32 CallstackId = ExternalCallstackId ? ExternalCallstackId : GDoNotAllocateInTrace ? 0 : CallstackTrace_GetCurrentId();
 
 	switch (RootHeap)
 	{
@@ -522,7 +522,7 @@ void MemoryTrace_Free(uint64 Address, HeapId RootHeap)
 }
 
 ////////////////////////////////////////////////////////////////////////////////
-void MemoryTrace_ReallocAlloc(uint64 Address, uint64 Size, uint32 Alignment, HeapId RootHeap)
+void MemoryTrace_ReallocAlloc(uint64 Address, uint64 Size, uint32 Alignment, HeapId RootHeap, uint32 ExternalCallstackId)
 {
 	if (!GTraceAllowed)
 	{
@@ -533,7 +533,7 @@ void MemoryTrace_ReallocAlloc(uint64 Address, uint64 Size, uint32 Alignment, Hea
 
 	const uint32 AlignmentPow2 = uint32(FPlatformMath::CountTrailingZeros(Alignment));
 	const uint32 Alignment_SizeLower = (AlignmentPow2 << SizeShift) | uint32(Size & ((1 << SizeShift) - 1));
-	const uint32 CallstackId = GDoNotAllocateInTrace ? 0 : CallstackTrace_GetCurrentId();
+	const uint32 CallstackId = ExternalCallstackId ? ExternalCallstackId : GDoNotAllocateInTrace ? 0 : CallstackTrace_GetCurrentId();
 
 	switch (RootHeap)
 	{
@@ -563,7 +563,7 @@ void MemoryTrace_ReallocAlloc(uint64 Address, uint64 Size, uint32 Alignment, Hea
 }
 
 ////////////////////////////////////////////////////////////////////////////////
-void MemoryTrace_ReallocFree(uint64 Address, HeapId RootHeap)
+void MemoryTrace_ReallocFree(uint64 Address, HeapId RootHeap, uint32 ExternalCallstackId)
 {
 	if (!GTraceAllowed)
 	{
@@ -572,7 +572,7 @@ void MemoryTrace_ReallocFree(uint64 Address, HeapId RootHeap)
 
 	checkSlow(RootHeap < 16);
 
-	const uint32 CallstackId = GDoNotAllocateInTrace ? 0 : CallstackTrace_GetCurrentId();
+	const uint32 CallstackId = ExternalCallstackId ? ExternalCallstackId : GDoNotAllocateInTrace ? 0 : CallstackTrace_GetCurrentId();
 
 	switch (RootHeap)
 	{
@@ -645,14 +645,14 @@ HeapId MemoryTrace_RootHeapSpec(const TCHAR* Name, EMemoryTraceHeapFlags Flags)
 }
 
 ////////////////////////////////////////////////////////////////////////////////
-void MemoryTrace_MarkAllocAsHeap(uint64 Address, HeapId Heap, EMemoryTraceHeapAllocationFlags Flags)
+void MemoryTrace_MarkAllocAsHeap(uint64 Address, HeapId Heap, EMemoryTraceHeapAllocationFlags Flags, uint32 ExternalCallstackId)
 {
 	if (!GTraceAllowed)
 	{
 		return;
 	}
 
-	const uint32 CallstackId = GDoNotAllocateInTrace ? 0 : CallstackTrace_GetCurrentId();
+	const uint32 CallstackId = ExternalCallstackId ? ExternalCallstackId : GDoNotAllocateInTrace ? 0 : CallstackTrace_GetCurrentId();
 
 	UE_TRACE_LOG(Memory, HeapMarkAlloc, MemAllocChannel)
 		<< HeapMarkAlloc.Address(uint64(Address))
@@ -662,14 +662,14 @@ void MemoryTrace_MarkAllocAsHeap(uint64 Address, HeapId Heap, EMemoryTraceHeapAl
 }
 
 ////////////////////////////////////////////////////////////////////////////////
-void MemoryTrace_UnmarkAllocAsHeap(uint64 Address, HeapId Heap)
+void MemoryTrace_UnmarkAllocAsHeap(uint64 Address, HeapId Heap, uint32 ExternalCallstackId)
 {
 	if (!GTraceAllowed)
 	{
 		return;
 	}
 
-	const uint32 CallstackId = GDoNotAllocateInTrace ? 0 : CallstackTrace_GetCurrentId();
+	const uint32 CallstackId = ExternalCallstackId ? ExternalCallstackId : GDoNotAllocateInTrace ? 0 : CallstackTrace_GetCurrentId();
 
 	// Sets all flags to zero
 	UE_TRACE_LOG(Memory, HeapUnmarkAlloc, MemAllocChannel)
diff --git a/Engine/Source/Runtime/Core/Public/ProfilingDebugging/MemoryTrace.h b/Engine/Source/Runtime/Core/Public/ProfilingDebugging/MemoryTrace.h
--- a/Engine/Source/Runtime/Core/Public/ProfilingDebugging/MemoryTrace.h
+++ b/Engine/Source/Runtime/Core/Public/ProfilingDebugging/MemoryTrace.h
@@ -92,16 +92,18 @@ CORE_API HeapId MemoryTrace_RootHeapSpec(const TCHAR* Name, EMemoryTraceHeapFlag
  * @param Address Address of the allocation
  * @param Heap Heap id, see /ref MemoryTrace_HeapSpec. If no specific heap spec has been created the correct root heap needs to be given.
  * @param Flags Additional properties of the heap allocation. Note that \ref EMemoryTraceHeapAllocationFlags::Heap is implicit.
+ * @param ExternalCallstackId CallstackId to use, if 0 will use current callstack id.
  */
-CORE_API void MemoryTrace_MarkAllocAsHeap(uint64 Address, HeapId Heap, EMemoryTraceHeapAllocationFlags Flags = EMemoryTraceHeapAllocationFlags::None);
+CORE_API void MemoryTrace_MarkAllocAsHeap(uint64 Address, HeapId Heap, EMemoryTraceHeapAllocationFlags Flags = EMemoryTraceHeapAllocationFlags::None, uint32 ExternalCallstackId = 0);
 
 /**
  * Unmark an allocation as a heap. When an allocation that has previously been used as a heap is reused as a regular
  * allocation.
  * @param Address Address of the allocation
  * @param Heap Heap id
+ * @param ExternalCallstackId CallstackId to use, if 0 will use current callstack id.
  */
-CORE_API void MemoryTrace_UnmarkAllocAsHeap(uint64 Address, HeapId Heap);
+CORE_API void MemoryTrace_UnmarkAllocAsHeap(uint64 Address, HeapId Heap, uint32 ExternalCallstackId = 0);
 
 /**
  * Trace an allocation event.
@@ -109,30 +111,34 @@ CORE_API void MemoryTrace_UnmarkAllocAsHeap(uint64 Address, HeapId Heap);
  * @param Size Size of allocation
  * @param Alignment Alignment of the allocation
  * @param RootHeap Which root heap this belongs to (system memory, video memory etc)
+ * @param ExternalCallstackId CallstackId to use, if 0 will use current callstack id.
  */
-CORE_API void MemoryTrace_Alloc(uint64 Address, uint64 Size, uint32 Alignment, HeapId RootHeap = EMemoryTraceRootHeap::SystemMemory);
+CORE_API void MemoryTrace_Alloc(uint64 Address, uint64 Size, uint32 Alignment, HeapId RootHeap = EMemoryTraceRootHeap::SystemMemory, uint32 ExternalCallstackId = 0);
 
 /**
  * Trace a free event.
  * @param Address Address of the allocation being freed
  * @param RootHeap Which root heap this belongs to (system memory, video memory etc)
+ * @param ExternalCallstackId CallstackId to use, if 0 will use current callstack id.
  */
-CORE_API void MemoryTrace_Free(uint64 Address, HeapId RootHeap = EMemoryTraceRootHeap::SystemMemory);
+CORE_API void MemoryTrace_Free(uint64 Address, HeapId RootHeap = EMemoryTraceRootHeap::SystemMemory, uint32 ExternalCallstackId = 0);
 
 /**
  * Trace a free related to a reallocation event.
  * @param Address Address of the allocation being freed
  * @param RootHeap Which root heap this belongs to (system memory, video memory etc)
+ * @param ExternalCallstackId CallstackId to use, if 0 will use current callstack id.
  */
-CORE_API void MemoryTrace_ReallocFree(uint64 Address, HeapId RootHeap = EMemoryTraceRootHeap::SystemMemory);
+CORE_API void MemoryTrace_ReallocFree(uint64 Address, HeapId RootHeap = EMemoryTraceRootHeap::SystemMemory, uint32 ExternalCallstackId = 0);
 
 /** Trace an allocation related to a reallocation event.
  * @param Address Address of allocation
  * @param NewSize Size of allocation
  * @param Alignment Alignment of the allocation
  * @param RootHeap Which root heap this belongs to (system memory, video memory etc)
+ * @param ExternalCallstackId CallstackId to use, if 0 will use current callstack id.
  */
-CORE_API void MemoryTrace_ReallocAlloc(uint64 Address, uint64 NewSize, uint32 Alignment, HeapId RootHeap = EMemoryTraceRootHeap::SystemMemory);
+CORE_API void MemoryTrace_ReallocAlloc(uint64 Address, uint64 NewSize, uint32 Alignment, HeapId RootHeap = EMemoryTraceRootHeap::SystemMemory, uint32 ExternalCallstackId = 0);
 
 ////////////////////////////////////////////////////////////////////////////////
 #else // UE_MEMORY_TRACE_ENABLED
@@ -142,9 +148,9 @@ inline HeapId MemoryTrace_RootHeapSpec(const TCHAR* Name, EMemoryTraceHeapFlags
 inline HeapId MemoryTrace_HeapSpec(HeapId ParentId, const TCHAR* Name, EMemoryTraceHeapFlags Flags = EMemoryTraceHeapFlags::None) { return ~0; }
 inline void MemoryTrace_MarkAllocAsHeap(uint64 Address, HeapId Heap) {}
 inline void MemoryTrace_UnmarkAllocAsHeap(uint64 Address, HeapId Heap) {}
-inline void MemoryTrace_Alloc(uint64 Address, uint64 Size, uint32 Alignment, HeapId RootHeap = EMemoryTraceRootHeap::SystemMemory) {}
-inline void MemoryTrace_Free(uint64 Address, HeapId RootHeap = EMemoryTraceRootHeap::SystemMemory) {}
-inline void MemoryTrace_ReallocFree(uint64 Address, HeapId RootHeap = EMemoryTraceRootHeap::SystemMemory) {}
-inline void MemoryTrace_ReallocAlloc(uint64 Address, uint64 NewSize, uint32 Alignment, HeapId RootHeap = EMemoryTraceRootHeap::SystemMemory) {}
+inline void MemoryTrace_Alloc(uint64 Address, uint64 Size, uint32 Alignment, HeapId RootHeap = EMemoryTraceRootHeap::SystemMemory, uint32 ExternalCallstackId = 0) {}
+inline void MemoryTrace_Free(uint64 Address, HeapId RootHeap = EMemoryTraceRootHeap::SystemMemory, uint32 ExternalCallstackId = 0) {}
+inline void MemoryTrace_ReallocFree(uint64 Address, HeapId RootHeap = EMemoryTraceRootHeap::SystemMemory, uint32 ExternalCallstackId = 0) {}
+inline void MemoryTrace_ReallocAlloc(uint64 Address, uint64 NewSize, uint32 Alignment, HeapId RootHeap = EMemoryTraceRootHeap::SystemMemory, uint32 ExternalCallstackId = 0) {}
 
 #endif // UE_MEMORY_TRACE_ENABLED
