Tracing Size == 0 events in libc.so malloc/free/etc memory tracer on Android

diff --git a/Engine/Build/Android/Prebuilt/ScudoMemoryTrace/ScudoMemoryTrace.cpp b/Engine/Build/Android/Prebuilt/ScudoMemoryTrace/ScudoMemoryTrace.cpp
--- a/Engine/Build/Android/Prebuilt/ScudoMemoryTrace/ScudoMemoryTrace.cpp
+++ b/Engine/Build/Android/Prebuilt/ScudoMemoryTrace/ScudoMemoryTrace.cpp
@@ -550,7 +550,7 @@ static inline void ReportAllocEventImpl(const void* FrameStartPtr, const void* R
 	}
 
 	const uint64_t Ptr = RemovePtrTags(RawPtr);
-	if (Ptr == 0 || Size == 0)
+	if (Ptr == 0)
 	{
 		return;
 	}
@@ -656,11 +656,15 @@ void ScudoMemoryTrace_SetHook(ScudoMemoryTraceHook TraceHook)
 
 extern "C" void __scudo_allocate_hook(void* Ptr, size_t Size)
 {
+	Setup();
+
 	ReportAllocEvent(Ptr, Size, 0);
 }
 
 extern "C" void __scudo_deallocate_hook(void* Ptr)
 {
+	Setup();
+
 	ReportFreeEvent(Ptr, 1);
 }
 
