Refactoring onCreate for GameActivity and switching nativemain to use AndroidMain, also updating ASISStub to properly handle activity change and theme change.
Also added gradle build files for ASISStub to make it easier to use the sample android app to test Android Single Instance
Also fixes for Jira: UE-205628
Jira UEHMI-316
#android

diff --git a/Engine/Build/Android/Java/src/com/epicgames/unreal/GameActivity.java.template b/Engine/Build/Android/Java/src/com/epicgames/unreal/GameActivity.java.template
--- a/Engine/Build/Android/Java/src/com/epicgames/unreal/GameActivity.java.template
+++ b/Engine/Build/Android/Java/src/com/epicgames/unreal/GameActivity.java.template
@@ -542,7 +542,7 @@ public class GameActivity extends $${gameActivitySuperClass}$$ implements Surfac
 	
 	public GameActivity()
 	{
-		Log.debug( "GameActivityForMakeAAR constructor called with _gameActivityInstance=" + _gameActivityInstance + ", activityContext=" + activityContext + ", this=" + this + ", isStandalone=" + isStandalone);
+		Log.debug( "GameActivity constructor called with _gameActivityInstance=" + _gameActivityInstance + ", activityContext=" + activityContext + ", this=" + this + ", isStandalone=" + isStandalone);
 
 		assert (_gameActivityInstance == null);
 		_gameActivityInstance = this;
@@ -568,6 +568,22 @@ public class GameActivity extends $${gameActivitySuperClass}$$ implements Surfac
 	{
 		return _gameActivityInstance.activityContext;
 	}
+
+	public final Window getWindow() 
+	{
+		return isStandalone ? activityContext.getWindow() : super.getWindow();
+	}
+
+	public final Intent getIntent()
+	{
+		return super.getIntent();
+	}
+
+	public final WindowManager getWindowManager()
+	{
+		return isStandalone ? activityContext.getWindowManager() : super.getWindowManager();
+	}
+
 	
 	/**
 	Get the SDK level of the OS we are running in.
@@ -603,9 +619,9 @@ public class GameActivity extends $${gameActivitySuperClass}$$ implements Surfac
 	public int getDeviceDefaultOrientation() 
 	{
 		// WindowManager windowManager =  (WindowManager) getSystemService(WINDOW_SERVICE);
-		WindowManager windowManager = activityContext.getWindowManager();
+		WindowManager windowManager = getWindowManager();
 
-		Configuration config = activityContext.getResources().getConfiguration();
+		Configuration config = getResources().getConfiguration();
 
 		int rotation = windowManager.getDefaultDisplay().getRotation();
 		if ( ((rotation == android.view.Surface.ROTATION_0 || rotation == android.view.Surface.ROTATION_180) &&
@@ -641,7 +657,7 @@ public class GameActivity extends $${gameActivitySuperClass}$$ implements Surfac
 	private int getResourceId(String VariableName, String ResourceName, String PackageName)
 	{
 		try {
-			return activityContext.getResources().getIdentifier(VariableName, ResourceName, PackageName);
+			return getResources().getIdentifier(VariableName, ResourceName, PackageName);
 		}
 		catch (Exception e) {
 			e.printStackTrace();
@@ -658,7 +674,7 @@ public class GameActivity extends $${gameActivitySuperClass}$$ implements Surfac
 	public String getProductName()
 	{
 	
-		String ProductName = Settings.Global.getString(activityContext.getApplicationContext().getContentResolver(), "default_device_name");
+		String ProductName = Settings.Global.getString(getApplicationContext().getContentResolver(), "default_device_name");
 		if(ProductName == null)
 		{
 			ProductName="";
@@ -715,7 +731,7 @@ public class GameActivity extends $${gameActivitySuperClass}$$ implements Surfac
 
     public boolean AndroidThunkJava_IsPackageInstalled(String packagename)
 	{
-        List<String> packages = getInstalledPackages(activityContext.getApplicationContext());
+        List<String> packages = getInstalledPackages(getApplicationContext());
         return packages.contains(packagename);
     }
 
@@ -1818,7 +1834,7 @@ public class GameActivity extends $${gameActivitySuperClass}$$ implements Surfac
 		int navigationBarHeight = 0;
 		if (Build.VERSION.SDK_INT >= 21)
 		{
-			android.content.res.Resources resources = activityContext.getApplicationContext().getResources();
+			android.content.res.Resources resources = getApplicationContext().getResources();
 			int resourceId = resources.getIdentifier("status_bar_height", "dimen", "android");
 			if (resourceId > 0) {
 				statusBarHeight += resources.getDimensionPixelSize(resourceId);
@@ -1832,7 +1848,7 @@ public class GameActivity extends $${gameActivitySuperClass}$$ implements Surfac
 		variables.put("navigationBarHeight", Integer.toString(navigationBarHeight));
 
 		// store the screen resolution (adjusted for immersive mode)
-		android.view.Display display = activityContext.getWindowManager().getDefaultDisplay();
+		android.view.Display display = getWindowManager().getDefaultDisplay();
 		android.graphics.Point displaySize = new android.graphics.Point();
 		if (ShouldHideUI)
 		{
@@ -2430,6 +2446,11 @@ public class GameActivity extends $${gameActivitySuperClass}$$ implements Surfac
 
 		super.onCreate(savedInstanceState);
 
+		onCreateBody(savedInstanceState);
+	}
+
+	protected void onCreateBody(Bundle savedInstanceState)
+	{
 		sensorManager = (SensorManager) getSystemService(Context.SENSOR_SERVICE);
 		accelerometer = sensorManager.getDefaultSensor(Sensor.TYPE_ACCELEROMETER);
 		magnetometer = sensorManager.getDefaultSensor(Sensor.TYPE_MAGNETIC_FIELD);
@@ -2447,7 +2468,7 @@ public class GameActivity extends $${gameActivitySuperClass}$$ implements Surfac
 			deviceVibrator = (Vibrator) getSystemService(Context.VIBRATOR_SERVICE);
 		}
 
-		if (PermissionHelper.checkPermission("android.permission.CHANGE_WIFI_MULTICAST_STATE"))
+		if (!isStandalone && PermissionHelper.checkPermission("android.permission.CHANGE_WIFI_MULTICAST_STATE"))
 		{
 			WifiManager wifiManager = (WifiManager)getSystemService(Context.WIFI_SERVICE);
 			MultiCastLock = wifiManager.createMulticastLock("UE");
@@ -2457,14 +2478,25 @@ public class GameActivity extends $${gameActivitySuperClass}$$ implements Surfac
 		// disable autofill to avoid possible Android 10 crash (android.view.autofill.AutofillManager.ensureServiceClientAddedIfNeededLocked() SyncResultReceiver$TimeoutException)
 		if (Build.VERSION.SDK_INT >= Build.VERSION_CODES.O)
 		{
-			getWindow().getDecorView().setImportantForAutofill(View.IMPORTANT_FOR_AUTOFILL_NO_EXCLUDE_DESCENDANTS);
+			try
+			{
+				getWindow().getDecorView().setImportantForAutofill(View.IMPORTANT_FOR_AUTOFILL_NO_EXCLUDE_DESCENDANTS);
+			} 
+			catch(Exception e) 
+			{
+				Log.error("onCreateBody(): getWindow().getDecorView() not available in current context!, err = " + e.getMessage());
+			}
 		}
 
 		boolean bIsShipping = nativeIsShippingBuild();
 
 		// create splashscreen dialog (if launched by SplashActivity)
-		_extrasBundle = getIntent().getExtras();
-		if (_extrasBundle != null)
+		if (!isStandalone && _extrasBundle == null)
+		{
+			_extrasBundle = getIntent().getExtras();
+		}
+
+		if (!isStandalone && _extrasBundle != null)
 		{
 			ShouldHideUI = _extrasBundle.getString("ShouldHideUI") != null;
 			UseDisplayCutout = _extrasBundle.getString("UseDisplayCutout") != null;
@@ -2473,7 +2505,7 @@ public class GameActivity extends $${gameActivitySuperClass}$$ implements Surfac
 				SplashScreenLaunch = true;
 				try {
 					// try to get the splash theme (can't use R.style.UnrealSplashTheme since we don't know the package name until runtime)
-					int SplashThemeId = activityContext.getResources().getIdentifier("UnrealSplashTheme", "style", getPackageName());
+					int SplashThemeId = getResources().getIdentifier("UnrealSplashTheme", "style", getPackageName());
 					mSplashDialog = new Dialog(this, SplashThemeId);
 					mSplashDialog.setCancelable(false);
 
@@ -2499,14 +2531,14 @@ public class GameActivity extends $${gameActivitySuperClass}$$ implements Surfac
 					e.printStackTrace();
 				}
 				try {
-					noActionAnimID = activityContext.getResources().getIdentifier("noaction", "anim", getPackageName());
+					noActionAnimID = getResources().getIdentifier("noaction", "anim", getPackageName());
 				}
 				catch (Exception e) {
 					e.printStackTrace();
 				}
 			}
 			try {
-				noActionAnimID = activityContext.getResources().getIdentifier("noaction", "anim", getPackageName());
+				noActionAnimID = getResources().getIdentifier("noaction", "anim", getPackageName());
 			}
 			catch (Exception e) {
 				e.printStackTrace();
@@ -2521,16 +2553,16 @@ public class GameActivity extends $${gameActivitySuperClass}$$ implements Surfac
 			if (UseDisplayCutout)
 			{
 				// will not be true if not Android Pie or later
-	            WindowManager.LayoutParams params = activityContext.getWindow().getAttributes();
+	            WindowManager.LayoutParams params = getWindow().getAttributes();
 		        params.layoutInDisplayCutoutMode = WindowManager.LayoutParams.LAYOUT_IN_DISPLAY_CUTOUT_MODE_SHORT_EDGES;
-			  	activityContext.getWindow().setAttributes(params);
+			  	getWindow().setAttributes(params);
 			}
 		}
 
 		// Retrieve version code, name, and target SDK
 		try 
 		{
-			PackageInfo packageInfo = activityContext.getPackageManager().getPackageInfo(getPackageName(), 0);
+			PackageInfo packageInfo = getPackageManager().getPackageInfo(getPackageName(), 0);
 			VersionCode = packageInfo.versionCode;
 			VersionName = packageInfo.versionName;
 			targetSdkVersion = packageInfo.applicationInfo.targetSdkVersion;
@@ -2541,7 +2573,7 @@ public class GameActivity extends $${gameActivitySuperClass}$$ implements Surfac
 		}
 
 		// Check for target sdk.  If 23 or higher then warn that permission handling may mean features don't work if user denies them.
-		if (ANDROID_BUILD_VERSION >= 23 && targetSdkVersion >= 23) //23 is the API level (Marshmallow) where runtime permission handling is available
+		if (!isStandalone && ANDROID_BUILD_VERSION >= 23 && targetSdkVersion >= 23) //23 is the API level (Marshmallow) where runtime permission handling is available
 		{
 			Log.debug("Target SDK is " + targetSdkVersion + ".  This may cause issues if permissions are denied by the user." );	
 			if (!bIsShipping)
@@ -2574,7 +2606,7 @@ public class GameActivity extends $${gameActivitySuperClass}$$ implements Surfac
 			{
 				int ProcessMemory = 0;
 
-				ActivityManager activityManager = (ActivityManager) activityContext.getSystemService(Context.ACTIVITY_SERVICE);
+				ActivityManager activityManager = (ActivityManager) getSystemService(Context.ACTIVITY_SERVICE);
 				int pid = android.os.Process.myPid();
 				int pids[] = new int[] { pid };
 				android.os.Debug.MemoryInfo[] memoryInfo = activityManager.getProcessMemoryInfo(pids);
@@ -2639,17 +2671,28 @@ public class GameActivity extends $${gameActivitySuperClass}$$ implements Surfac
 				}
 			}
 		};
+
 		memoryHandlerThread = new HandlerThread("MemoryUsageThread", android.os.Process.THREAD_PRIORITY_LOWEST);
 		memoryHandlerThread.start();
 		memoryHandler = new Handler(memoryHandlerThread.getLooper());
 		memoryHandler.postDelayed(memoryRunnable, 1);
 
 		// layout required by popups, e.g ads, native controls
-		MarginLayoutParams params = new MarginLayoutParams(LayoutParams.WRAP_CONTENT, LayoutParams.WRAP_CONTENT);
-		params.setMargins(0, 0, 0, 0);
-		activityLayout = new LinearLayout(activityContext);
-		SimpleContextWrapper.getActivityOfContext(activityContext).setContentView(activityLayout, params);
+		if (!isStandalone)
+		{
+			MarginLayoutParams params = new MarginLayoutParams(LayoutParams.WRAP_CONTENT, LayoutParams.WRAP_CONTENT);
+			params.setMargins(0, 0, 0, 0);
+			activityLayout = new LinearLayout(activityContext);
+			SimpleContextWrapper.getActivityOfContext(activityContext).setContentView(activityLayout, params);
 
+			// tell Android that we want volume controls to change the media volume, aka music
+			setVolumeControlStream(AudioManager.STREAM_MUSIC);
+		}
+
+		//
+		// NOTE: DO NOT DELETE THIS!!!!!!!
+		//       this is meant to be uncommented for special cases
+		//
 /*
 		// Turn on and unlock screen.. Assumption is that this
 		// will only really have an effect when for debug launching
@@ -2669,6 +2712,10 @@ public class GameActivity extends $${gameActivitySuperClass}$$ implements Surfac
 		}
 */
 
+		//
+		// NOTE: DO NOT DELETE THIS!!!!!!!
+		//       this is meant to be uncommented for special cases
+		//
 /*
 		// log a list of input devices for debugging
 		{
@@ -2690,13 +2737,10 @@ public class GameActivity extends $${gameActivitySuperClass}$$ implements Surfac
 		}
 */
 
-		// tell Android that we want volume controls to change the media volume, aka music
-		setVolumeControlStream(AudioManager.STREAM_MUSIC);
-
 		// Look for Vulkan support if Nougat or later
 		if (ANDROID_BUILD_VERSION >= 24)
 		{
-			FeatureInfo[] features = activityContext.getPackageManager().getSystemAvailableFeatures();
+			FeatureInfo[] features = getPackageManager().getSystemAvailableFeatures();
 			for (FeatureInfo feature : features) {
 				if (feature.name != null)
 				{
@@ -2747,7 +2791,7 @@ public class GameActivity extends $${gameActivitySuperClass}$$ implements Surfac
 			boolean bForceLandscape = false;
 
 			// check for a Google TV by checking system feature support
-			if (activityContext.getPackageManager().hasSystemFeature("com.google.android.tv"))
+			if (getPackageManager().hasSystemFeature("com.google.android.tv"))
 			{
 				Log.debug( "Detected Google TV, will default to landscape" );
 				bForceLandscape = true;
@@ -2800,7 +2844,7 @@ public class GameActivity extends $${gameActivitySuperClass}$$ implements Surfac
 		ProjectName = ProjectName.substring(ProjectName.lastIndexOf('.') + 1);
 		String AppType = "";
 		try {
-			ApplicationInfo ai = activityContext.getPackageManager().getApplicationInfo(getPackageName(), PackageManager.GET_META_DATA);
+			ApplicationInfo ai = getPackageManager().getApplicationInfo(getPackageName(), PackageManager.GET_META_DATA);
 			Bundle bundle = ai.metaData;
 			_bundle = bundle;
 
@@ -2991,7 +3035,7 @@ public class GameActivity extends $${gameActivitySuperClass}$$ implements Surfac
 		nativeSetGlobalActivity(UseExternalFilesDir, PublicLogFiles, getFilesDir().getPath(), getExternalFilesDir(null).getPath(), PackageDataInsideApkValue==1, APKPath);
 
 		// tell the engine if this is a portrait app
-		nativeSetWindowInfo(activityContext.getResources().getConfiguration().orientation == Configuration.ORIENTATION_PORTRAIT, DepthBufferPreference, PropagateAlpha);
+		nativeSetWindowInfo(getResources().getConfiguration().orientation == Configuration.ORIENTATION_PORTRAIT, DepthBufferPreference, PropagateAlpha);
 
 		// get the full language code, like en-US
 		// note: this may need to be Locale.getDefault().getLanguage()
@@ -3013,7 +3057,7 @@ public class GameActivity extends $${gameActivitySuperClass}$$ implements Surfac
 		int PatchVersion = 0;
 		if (VersionCode > 0)
 		{
-			nativeSetObbInfo(ProjectName, activityContext.getApplicationContext().getPackageName(), VersionCode, PatchVersion, AppType);
+			nativeSetObbInfo(ProjectName, getApplicationContext().getPackageName(), VersionCode, PatchVersion, AppType);
 		}
 		else
 		{
@@ -3022,14 +3066,17 @@ public class GameActivity extends $${gameActivitySuperClass}$$ implements Surfac
 		}
 		
 		// enable the physical volume controls to the game
-		this.setVolumeControlStream(AudioManager.STREAM_MUSIC);
+		if (!isStandalone)
+		{
+			this.setVolumeControlStream(AudioManager.STREAM_MUSIC);
+		}
 
 		AlertDialog.Builder builder;
 
-		consoleInputBox = new ConsoleKeyboardInput(this);
+		consoleInputBox = new ConsoleKeyboardInput(activityContext);
 		consoleInputBox.setInputType(0x00080001); // TYPE_CLASS_TEXT | TYPE_TEXT_FLAG_NO_SUGGESTIONS);
 		
-		RelativeLayout consoleInputLayout = new RelativeLayout(this);
+		RelativeLayout consoleInputLayout = new RelativeLayout(activityContext);
 		{ 
 			android.widget.ImageButton consoleHistoryButton = consoleInputBox.getButton();
 			consoleInputBox.setId(1);
@@ -3047,8 +3094,8 @@ public class GameActivity extends $${gameActivitySuperClass}$$ implements Surfac
 		}
 
 		// Spinner with Quick Stat Commands
-		consoleSpinner = new Spinner(this);
-		ArrayAdapter<String> adapter = new ArrayAdapter<String>(this, android.R.layout.simple_spinner_item, CONSOLE_SPINNER_ITEMS);
+		consoleSpinner = new Spinner(activityContext);
+		ArrayAdapter<String> adapter = new ArrayAdapter<String>(activityContext, android.R.layout.simple_spinner_item, CONSOLE_SPINNER_ITEMS);
 		adapter.setDropDownViewResource(android.R.layout.simple_spinner_dropdown_item);
 		consoleSpinner.setAdapter(adapter);
 		consoleSpinner.setOnItemSelectedListener(new AdapterView.OnItemSelectedListener() {
@@ -3066,12 +3113,12 @@ public class GameActivity extends $${gameActivitySuperClass}$$ implements Surfac
 		});
 
 		// Layout for Quick Commands and Console Input
-		consoleAlertLayout = new LinearLayout(this);
+		consoleAlertLayout = new LinearLayout(activityContext);
 		consoleAlertLayout.setOrientation(LinearLayout.VERTICAL);
 		consoleAlertLayout.addView(consoleSpinner);
 		consoleAlertLayout.addView(consoleInputLayout);
 
-		builder = new AlertDialog.Builder(this);
+		builder = new AlertDialog.Builder(activityContext);
 		builder.setTitle("Console Window - Enter Command")
 		.setMessage("")
 		.setView(consoleAlertLayout)
@@ -3102,7 +3149,7 @@ public class GameActivity extends $${gameActivitySuperClass}$$ implements Surfac
 		});
 		consoleAlert = builder.create();
 
-		virtualKeyboardInputBox = new EditText(this);
+		virtualKeyboardInputBox = new EditText(activityContext);
 		virtualKeyboardInputBox.addTextChangedListener(new TextWatcher() {
 			@Override
 			public void beforeTextChanged(CharSequence charSequence, int start, int count, int after) {
@@ -3119,7 +3166,7 @@ public class GameActivity extends $${gameActivitySuperClass}$$ implements Surfac
 			}
 		});
 
-		builder = new AlertDialog.Builder(this);
+		builder = new AlertDialog.Builder(activityContext);
 		builder.setTitle("")
 		.setView(virtualKeyboardInputBox)
 		.setPositiveButton("Ok", new DialogInterface.OnClickListener() {
@@ -3142,8 +3189,11 @@ public class GameActivity extends $${gameActivitySuperClass}$$ implements Surfac
 		});
 		virtualKeyboardAlert = builder.create();
 
-		GooglePlayLicensing.GoogleLicensing = new GooglePlayLicensing();
-		GooglePlayLicensing.GoogleLicensing.Init(this, Log);
+		if (!isStandalone) // TODO: maybe this is okay for ASIS
+		{
+			GooglePlayLicensing.GoogleLicensing = new GooglePlayLicensing();
+			GooglePlayLicensing.GoogleLicensing.Init(this, Log);
+		}
 
 		// Now okay for event handler to be set up on native side
 		//	nativeResumeMainInit();
@@ -3152,18 +3202,14 @@ public class GameActivity extends $${gameActivitySuperClass}$$ implements Surfac
 		// AndroidThunkJava_GooglePlayConnect();
 
 		// If we have data in the apk or just loose then carry on init as normal
-		/*Log.debug(this.getObbDir().getAbsolutePath());
-		String path = this.getObbDir().getAbsolutePath() + "/main.1.com.epicgames.StrategyGame.obb";
-		File obb = new File(path);
-		Log.debug("=+=+=+=+=+=+=> File exists: " + (obb.exists() ? "True" : "False"));
-		*/
+		
 		if(PackageDataInsideApkValue == 1 || HasOBBFiles == 0)
 		{
 			HasAllFiles = true;
 		}
 
 		// check for OBB file present if we don't have all the files and don't need to verify
-		if (!HasAllFiles && !VerifyOBBOnStartUp)
+		if (!isStandalone && !HasAllFiles && !VerifyOBBOnStartUp)
 		{
 			HasAllFiles = DownloadShim.expansionFilesDelivered(this, VersionCode);
 		}
@@ -3171,157 +3217,162 @@ public class GameActivity extends $${gameActivitySuperClass}$$ implements Surfac
 		containerFrameLayout = new FrameLayout(activityContext);
 		virtualKeyboardLayout = new LinearLayout(activityContext);
 
-		// Need to create our surface view here regardless of if we are going to end up using it
-		getWindow().takeSurface(null);
-		MySurfaceView = new SurfaceView(this);
-		MySurfaceView.setBackgroundColor(Color.TRANSPARENT);
-		MySurfaceView.getHolder().addCallback(this);
-		containerFrameLayout.addView(MySurfaceView);
-		containerFrameLayout.addView(virtualKeyboardLayout);
-		setContentView(containerFrameLayout);
-
-		// cache a reference to the main content view and set it so it can be focused on
-        mainView = findViewById( android.R.id.content );
-        mainView.setFocusable( true );
-        mainView.setFocusableInTouchMode( true );
-
-        mainDecorView = activityContext.getWindow().getDecorView();
-		mainDecorViewRect = new Rect();
-		
-		// check for display cutouts to set safe zone
-		View mainRootView = mainDecorView.findViewById(android.R.id.content).getRootView();
-		if (Build.VERSION.SDK_INT > 19)
+		if (!isStandalone)
 		{
-			mainRootView.setOnApplyWindowInsetsListener(new View.OnApplyWindowInsetsListener() {
-				@Override
-       			public WindowInsets onApplyWindowInsets(View v, WindowInsets insets) {
-					if (Build.VERSION.SDK_INT >= 28)
-					{
-						DisplayCutout cutout = insets.getDisplayCutout();
-						if (cutout != null)
+			//TODO: ASIS need to refactor how MySurfaceView is used since the surface can be external and also can be a TextureView
+
+			// Need to create our surface view here regardless of if we are going to end up using it
+			getWindow().takeSurface(null);
+			MySurfaceView = new SurfaceView(this);
+			MySurfaceView.setBackgroundColor(Color.TRANSPARENT);
+			MySurfaceView.getHolder().addCallback(this);
+			containerFrameLayout.addView(MySurfaceView);
+			containerFrameLayout.addView(virtualKeyboardLayout);
+			setContentView(containerFrameLayout);
+
+			// cache a reference to the main content view and set it so it can be focused on
+			mainView = findViewById( android.R.id.content );
+			mainView.setFocusable( true );
+			mainView.setFocusableInTouchMode( true );
+
+			mainDecorView = getWindow().getDecorView();
+			mainDecorViewRect = new Rect();
+		
+			// check for display cutouts to set safe zone
+			View mainRootView = mainDecorView.findViewById(android.R.id.content).getRootView();
+			if (Build.VERSION.SDK_INT > 19)
+			{
+				mainRootView.setOnApplyWindowInsetsListener(new View.OnApplyWindowInsetsListener() {
+					@Override
+       				public WindowInsets onApplyWindowInsets(View v, WindowInsets insets) {
+						if (Build.VERSION.SDK_INT >= 28)
 						{
-							Rect fullViewRect = new Rect();
-							mainDecorView.getDrawingRect( mainDecorViewRect );
-
-							android.view.Display display = activityContext.getWindowManager().getDefaultDisplay();
-							android.graphics.Point displaySize = new android.graphics.Point();
-							display.getRealSize(displaySize);
-
-							// store safezone independent of resolution
-							float invX = 1.0f / (displaySize.x != 0.0f ? displaySize.x : 1.0f);
-							float invY = 1.0f / (displaySize.y != 0.0f ? displaySize.y : 1.0f);
-							safezoneF.left = (safezone.left = cutout.getSafeInsetLeft()) * invX;
-							safezoneF.top = (safezone.top = cutout.getSafeInsetTop()) * invY;
-							safezoneF.right = (safezone.right = cutout.getSafeInsetRight()) * invX;
-							safezoneF.bottom = (safezone.bottom = cutout.getSafeInsetBottom()) *invY;
-							nativeSetSafezoneInfo(displaySize.y > displaySize.x, safezoneF.left, safezoneF.top, safezoneF.right, safezoneF.bottom);
-
-							//Log.debug("FullViewRect: " + fullViewRect.toString());	
-							//Log.debug("DisplayCutout: " + cutout.toString());
-							//Log.debug("SafeTop: " + cutout.getSafeInsetTop());
-							//Log.debug("SafeBottom: " + cutout.getSafeInsetBottom());
-							//Log.debug("SafeLeft: " + cutout.getSafeInsetLeft());
-							//Log.debug("SafeRight: " + cutout.getSafeInsetRight());
-							//List<Rect>bounds = cutout.getBoundingRects();
-							//for (Rect r : bounds)
-							//{
-							//	Log.debug("Rect: " + r.toString());
-							//}
+							DisplayCutout cutout = insets.getDisplayCutout();
+							if (cutout != null)
+							{
+								Rect fullViewRect = new Rect();
+								mainDecorView.getDrawingRect( mainDecorViewRect );
+
+								android.view.Display display = getWindowManager().getDefaultDisplay();
+								android.graphics.Point displaySize = new android.graphics.Point();
+								display.getRealSize(displaySize);
+
+								// store safezone independent of resolution
+								float invX = 1.0f / (displaySize.x != 0.0f ? displaySize.x : 1.0f);
+								float invY = 1.0f / (displaySize.y != 0.0f ? displaySize.y : 1.0f);
+								safezoneF.left = (safezone.left = cutout.getSafeInsetLeft()) * invX;
+								safezoneF.top = (safezone.top = cutout.getSafeInsetTop()) * invY;
+								safezoneF.right = (safezone.right = cutout.getSafeInsetRight()) * invX;
+								safezoneF.bottom = (safezone.bottom = cutout.getSafeInsetBottom()) *invY;
+								nativeSetSafezoneInfo(displaySize.y > displaySize.x, safezoneF.left, safezoneF.top, safezoneF.right, safezoneF.bottom);
+
+								//Log.debug("FullViewRect: " + fullViewRect.toString());	
+								//Log.debug("DisplayCutout: " + cutout.toString());
+								//Log.debug("SafeTop: " + cutout.getSafeInsetTop());
+								//Log.debug("SafeBottom: " + cutout.getSafeInsetBottom());
+								//Log.debug("SafeLeft: " + cutout.getSafeInsetLeft());
+								//Log.debug("SafeRight: " + cutout.getSafeInsetRight());
+								//List<Rect>bounds = cutout.getBoundingRects();
+								//for (Rect r : bounds)
+								//{
+								//	Log.debug("Rect: " + r.toString());
+								//}
+							}
 						}
+						return insets;
 					}
-					return insets;
-				}
-			});
-		}
-
-		createVirtualKeyboardInput();
+				});
+			}
 		
-		mainView.getViewTreeObserver().addOnGlobalLayoutListener( new ViewTreeObserver.OnGlobalLayoutListener()
-        {
-        	@Override
-        	public void onGlobalLayout()
-        	{
-				//Log.debug("VK: onGlobalLayout " + bKeyboardShowing);
-        		//if( bKeyboardShowing )
+			createVirtualKeyboardInput();
+		
+			mainView.getViewTreeObserver().addOnGlobalLayoutListener( new ViewTreeObserver.OnGlobalLayoutListener()
+			{
+        		@Override
+        		public void onGlobalLayout()
         		{
-        			Rect visibleRect = new Rect();
-        			View visibleView = mainView.getRootView();
+					//Log.debug("VK: onGlobalLayout " + bKeyboardShowing);
+        			//if( bKeyboardShowing )
+        			{
+        				Rect visibleRect = new Rect();
+        				View visibleView = mainView.getRootView();
 
-        			visibleView.getWindowVisibleDisplayFrame( visibleRect );
+        				visibleView.getWindowVisibleDisplayFrame( visibleRect );
 
-					mainDecorView.getDrawingRect( mainDecorViewRect );
+						mainDecorView.getDrawingRect( mainDecorViewRect );
 
-					//Log.debug("VK: onGlobalLayout visibleRect:(" +  visibleRect.left + ", " + visibleRect.top +", " + visibleRect.right +", " + visibleRect.bottom +")"+
-        			//", mainDecorViewRect:" +  mainDecorViewRect.left + ", " + mainDecorViewRect.top +", " + mainDecorViewRect.right +", " + mainDecorViewRect.bottom + ")" );
+						//Log.debug("VK: onGlobalLayout visibleRect:(" +  visibleRect.left + ", " + visibleRect.top +", " + visibleRect.right +", " + visibleRect.bottom +")"+
+        				//", mainDecorViewRect:" +  mainDecorViewRect.left + ", " + mainDecorViewRect.top +", " + mainDecorViewRect.right +", " + mainDecorViewRect.bottom + ")" );
 
-        			// determine which side of the screen the keyboard is covering
-        			int leftDiff = Math.abs( mainDecorViewRect.left - visibleRect.left );
-        			int topDiff = Math.abs( mainDecorViewRect.top - visibleRect.top );
-        			int rightDiff = Math.abs( mainDecorViewRect.right - visibleRect.right );
-        			int bottomDiff = Math.abs( mainDecorViewRect.bottom - visibleRect.bottom );
+        				// determine which side of the screen the keyboard is covering
+        				int leftDiff = Math.abs( mainDecorViewRect.left - visibleRect.left );
+        				int topDiff = Math.abs( mainDecorViewRect.top - visibleRect.top );
+        				int rightDiff = Math.abs( mainDecorViewRect.right - visibleRect.right );
+        				int bottomDiff = Math.abs( mainDecorViewRect.bottom - visibleRect.bottom );
 
-        			// Rect covered by the virtual keyboard
-        			Rect keyboardRect = new Rect();
-    				keyboardRect.left = ( rightDiff > 0 ) ? visibleRect.right : mainDecorViewRect.left; // keyboard is on the right
-    				keyboardRect.top = ( bottomDiff > 0 ) ? visibleRect.bottom : mainDecorViewRect.top; // keyboard is on the bottom
-    				keyboardRect.right = ( leftDiff > 0 ) ? visibleRect.left : mainDecorViewRect.right; // keyboard is on the left
-    				keyboardRect.bottom = ( topDiff > 0 ) ? visibleRect.top : mainDecorViewRect.bottom; // keyboard is on the top
+        				// Rect covered by the virtual keyboard
+        				Rect keyboardRect = new Rect();
+    					keyboardRect.left = ( rightDiff > 0 ) ? visibleRect.right : mainDecorViewRect.left; // keyboard is on the right
+    					keyboardRect.top = ( bottomDiff > 0 ) ? visibleRect.bottom : mainDecorViewRect.top; // keyboard is on the bottom
+    					keyboardRect.right = ( leftDiff > 0 ) ? visibleRect.left : mainDecorViewRect.right; // keyboard is on the left
+    					keyboardRect.bottom = ( topDiff > 0 ) ? visibleRect.top : mainDecorViewRect.bottom; // keyboard is on the top
 
-					//keyboard Y coord
-					int keyboardYPos = visibleRect.bottom - newVirtualKeyboardInput.getHeight();
+						//keyboard Y coord
+						int keyboardYPos = visibleRect.bottom - newVirtualKeyboardInput.getHeight();
 
-					//avoid negative coords if the keyboard is shown on top of the screen
-					if(keyboardYPos < 0)
-						keyboardYPos = visibleRect.top + newVirtualKeyboardInput.getHeight();
+						//avoid negative coords if the keyboard is shown on top of the screen
+						if(keyboardYPos < 0)
+							keyboardYPos = visibleRect.top + newVirtualKeyboardInput.getHeight();
 
-                    int visibleScreenYOffset = Math.max(bottomDiff, topDiff);
+						int visibleScreenYOffset = Math.max(bottomDiff, topDiff);
 
-					//Log.debug("VK: show?" + visibleScreenYOffset + "," + newVirtualKeyboardInput.getY());
+						//Log.debug("VK: show?" + visibleScreenYOffset + "," + newVirtualKeyboardInput.getY());
 					
-					// If the keyboard is a soft keyboard, the visibleScreenYOffset should be greater than 200 tall
-					boolean bSoftKeyboardPopup = visibleScreenYOffset > 200;
+						// If the keyboard is a soft keyboard, the visibleScreenYOffset should be greater than 200 tall
+						boolean bSoftKeyboardPopup = visibleScreenYOffset > 200;
 
-					// If the keyboard is a floating keyboard, after testing with some Chinese input methods the visibleScreenYOffset can be 0 pixel tall, expand it to [0, 200] for safety, correct this range if there is any unexpected case.
-					boolean bFloatingKeyboardPopup = visibleScreenYOffset >= 0 && visibleScreenYOffset <= 200;
+						// If the keyboard is a floating keyboard, after testing with some Chinese input methods the visibleScreenYOffset can be 0 pixel tall, expand it to [0, 200] for safety, correct this range if there is any unexpected case.
+						boolean bFloatingKeyboardPopup = visibleScreenYOffset >= 0 && visibleScreenYOffset <= 200;
 
-					if(bSoftKeyboardPopup || bFloatingKeyboardPopup) // actually this condition equals to the bKeyboardShowing, keep it in case we need to handle the different keyboards seperately.
-                    {
-						if (bKeyboardShowing)
+						if(bSoftKeyboardPopup || bFloatingKeyboardPopup) // actually this condition equals to the bKeyboardShowing, keep it in case we need to handle the different keyboards seperately.
 						{
-							nativeVirtualKeyboardShown( keyboardRect.left, keyboardRect.top, keyboardRect.right, keyboardRect.bottom );
-
-							//Log.debug("VK: show");
-							//#jira UE-55117 Android virtual keyboard can have text input hidden by software buttons
-							newVirtualKeyboardInput.getLayoutParams().width = Math.abs(visibleRect.right - visibleRect.left);
-							newVirtualKeyboardInput.setX(leftDiff);
-							newVirtualKeyboardInput.setY(keyboardYPos);
-							newVirtualKeyboardInput.setVisibility(View.VISIBLE);
-							newVirtualKeyboardInput.requestFocus();
-						}
+							if (bKeyboardShowing)
+							{
+								nativeVirtualKeyboardShown( keyboardRect.left, keyboardRect.top, keyboardRect.right, keyboardRect.bottom );
+
+								//Log.debug("VK: show");
+								//#jira UE-55117 Android virtual keyboard can have text input hidden by software buttons
+								newVirtualKeyboardInput.getLayoutParams().width = Math.abs(visibleRect.right - visibleRect.left);
+								newVirtualKeyboardInput.setX(leftDiff);
+								newVirtualKeyboardInput.setY(keyboardYPos);
+								newVirtualKeyboardInput.setVisibility(View.VISIBLE);
+								newVirtualKeyboardInput.requestFocus();
+							}
 
-						AndroidThunkJava_ResizeKeyboard(true, keyboardRect, visibleRect);
-                    }
-                    else
-                    {
-						if (bKeyboardShowing)
+							AndroidThunkJava_ResizeKeyboard(true, keyboardRect, visibleRect);
+						}
+						else
 						{
-							nativeVirtualKeyboardShown( keyboardRect.left, keyboardRect.top, keyboardRect.right, keyboardRect.bottom );
-							
-							if (newVirtualKeyboardInput.getY() > 0)
+							if (bKeyboardShowing)
 							{
+								nativeVirtualKeyboardShown( keyboardRect.left, keyboardRect.top, keyboardRect.right, keyboardRect.bottom );
+							
+								if (newVirtualKeyboardInput.getY() > 0)
+								{
 
-								//Log.debug("VK: hide");
-								newVirtualKeyboardInput.setVisibility(View.GONE);
-								//set offscreen
-								newVirtualKeyboardInput.setY(-1000);
+									//Log.debug("VK: hide");
+									newVirtualKeyboardInput.setVisibility(View.GONE);
+									//set offscreen
+									newVirtualKeyboardInput.setY(-1000);
+								}
 							}
-						}
 
-						AndroidThunkJava_ResizeKeyboard(false, keyboardRect, visibleRect);
-                    }
+							AndroidThunkJava_ResizeKeyboard(false, keyboardRect, visibleRect);
+						}
+        			}
         		}
-        	}
-        });
+			});
+		}
 
 //$${gameActivityAfterMainViewCreatedAdditions}$$
 
@@ -3346,6 +3397,7 @@ public class GameActivity extends $${gameActivitySuperClass}$$ implements Surfac
 
 //$${gameActivityOnCreateAdditions}$$
 //$${gameActivityOnCreateFinalAdditions}$$
+
 		mPSOProgramServiceAccessor.Init(this);
 
 		Log.debug("==============> GameActive.onCreate complete!");
@@ -3381,7 +3433,7 @@ public class GameActivity extends $${gameActivitySuperClass}$$ implements Surfac
 	{
 		if(android.os.Build.VERSION.SDK_INT >= 19) {
 			try {
-				View decorView = activityContext.getWindow().getDecorView(); 
+				View decorView = getWindow().getDecorView(); 
 
 				Log.debug("=== Restoring Transparent Bars ===");
 				// Clear the flag and then restore it
@@ -3440,13 +3492,13 @@ public class GameActivity extends $${gameActivitySuperClass}$$ implements Surfac
 		}
 
 		// invalidate window cache
-		nativeSetWindowInfo(activityContext.getResources().getConfiguration().orientation == Configuration.ORIENTATION_PORTRAIT, DepthBufferPreference, PropagateAlpha);
+		nativeSetWindowInfo(getResources().getConfiguration().orientation == Configuration.ORIENTATION_PORTRAIT, DepthBufferPreference, PropagateAlpha);
 
 		// only do this on KitKat and above
-		if (ShouldHideUI)
+		if (!isStandalone && ShouldHideUI)
 		{ 
 			restoreTransparentBars();
-			View decorView = activityContext.getWindow().getDecorView(); 
+			View decorView = getWindow().getDecorView(); 
 
 			decorView.setOnSystemUiVisibilityChangeListener(new View.OnSystemUiVisibilityChangeListener() {
 				@Override 
@@ -3468,10 +3520,10 @@ public class GameActivity extends $${gameActivitySuperClass}$$ implements Surfac
 		}
 
 		// only true if Android Pie or later
-		if (UseDisplayCutout)
+		if (!isStandalone && UseDisplayCutout)
 		{
 			// only on Android Pie and later
-			WindowManager.LayoutParams params = activityContext.getWindow().getAttributes();
+			WindowManager.LayoutParams params = getWindow().getAttributes();
 			if (params.layoutInDisplayCutoutMode != WindowManager.LayoutParams.LAYOUT_IN_DISPLAY_CUTOUT_MODE_SHORT_EDGES)
 			{
 				params.layoutInDisplayCutoutMode = WindowManager.LayoutParams.LAYOUT_IN_DISPLAY_CUTOUT_MODE_SHORT_EDGES;
@@ -3538,10 +3590,10 @@ public class GameActivity extends $${gameActivitySuperClass}$$ implements Surfac
 		}
 
 		// only do this on KitKat and above
-		if (ShouldHideUI)
+		if (!isStandalone && ShouldHideUI)
 		{ 
 			restoreTransparentBars();
-			View decorView = activityContext.getWindow().getDecorView(); 
+			View decorView = getWindow().getDecorView(); 
 
 			decorView.setOnSystemUiVisibilityChangeListener(new View.OnSystemUiVisibilityChangeListener() {
 				@Override 
@@ -3563,9 +3615,9 @@ public class GameActivity extends $${gameActivitySuperClass}$$ implements Surfac
 		}
 
 		// only true if Android Pie or later
-		if (UseDisplayCutout)
+		if (!isStandalone && UseDisplayCutout)
 		{
-			WindowManager.LayoutParams params = activityContext.getWindow().getAttributes();
+			WindowManager.LayoutParams params = getWindow().getAttributes();
 			if (params.layoutInDisplayCutoutMode != WindowManager.LayoutParams.LAYOUT_IN_DISPLAY_CUTOUT_MODE_SHORT_EDGES)
 			{
 				params.layoutInDisplayCutoutMode = WindowManager.LayoutParams.LAYOUT_IN_DISPLAY_CUTOUT_MODE_SHORT_EDGES;
@@ -4016,7 +4068,7 @@ public class GameActivity extends $${gameActivitySuperClass}$$ implements Surfac
 				public void run()
 				{
 					Log.debug("==============> [JAVA] AndroidThunkJava_KeepScreenOn(true) - Disabled screen saver");
-					activityContext.getWindow().addFlags(WindowManager.LayoutParams.FLAG_KEEP_SCREEN_ON);
+					getWindow().addFlags(WindowManager.LayoutParams.FLAG_KEEP_SCREEN_ON);
 				}
 			});
 		}
@@ -4028,7 +4080,7 @@ public class GameActivity extends $${gameActivitySuperClass}$$ implements Surfac
 				public void run()
 				{
 					Log.debug("==============> [JAVA] AndroidThunkJava_KeepScreenOn(false) - Enabled screen saver");
-					activityContext.getWindow().clearFlags(WindowManager.LayoutParams.FLAG_KEEP_SCREEN_ON);
+					getWindow().clearFlags(WindowManager.LayoutParams.FLAG_KEEP_SCREEN_ON);
 				}
 			});
 		}
@@ -4054,7 +4106,7 @@ public class GameActivity extends $${gameActivitySuperClass}$$ implements Surfac
 				public void run()
 				{
 					Log.debug("==============> [JAVA] AndroidThunkJava_DisableScreenCapture(true) - Disabled screen captures");
-					activityContext.getWindow().addFlags(WindowManager.LayoutParams.FLAG_SECURE);
+					getWindow().addFlags(WindowManager.LayoutParams.FLAG_SECURE);
 				}
 			});
 		}
@@ -4066,7 +4118,7 @@ public class GameActivity extends $${gameActivitySuperClass}$$ implements Surfac
 				public void run()
 				{
 					Log.debug("==============> [JAVA] AndroidThunkJava_DisableScreenCapture(false) - Enabled screen captures");
-					activityContext.getWindow().clearFlags(WindowManager.LayoutParams.FLAG_SECURE);
+					getWindow().clearFlags(WindowManager.LayoutParams.FLAG_SECURE);
 				}
 			});
 		}
@@ -4391,7 +4443,7 @@ public class GameActivity extends $${gameActivitySuperClass}$$ implements Surfac
 	private boolean AndroidThunkJava_isIntentActionAvailable(String intentAction)
 	{
 		Intent intent = new Intent(intentAction, null);
-		List<android.content.pm.ResolveInfo> activities = activityContext.getPackageManager().queryIntentActivities(intent, PackageManager.MATCH_DEFAULT_ONLY);
+		List<android.content.pm.ResolveInfo> activities = getPackageManager().queryIntentActivities(intent, PackageManager.MATCH_DEFAULT_ONLY);
 	    return activities != null && activities.size() > 0;
 	}
 
@@ -4438,7 +4490,7 @@ public class GameActivity extends $${gameActivitySuperClass}$$ implements Surfac
 			BrowserIntent.addFlags(Intent.FLAG_ACTIVITY_NEW_TASK | Intent.FLAG_ACTIVITY_MULTIPLE_TASK);
 			
 			// make sure there is a web browser to handle the URL before trying to start activity (or may crash!)
-			if (BrowserIntent.resolveActivity(activityContext.getPackageManager()) != null)
+			if (BrowserIntent.resolveActivity(getPackageManager()) != null)
 			{
 				Log.debug("[JAVA} AndroidThunkJava_LaunchURL: Starting activity");
 				startActivity(BrowserIntent);
@@ -4457,7 +4509,7 @@ public class GameActivity extends $${gameActivitySuperClass}$$ implements Surfac
 	public String AndroidThunkJava_GetAndroidId()
 	{
 		try {
-			String androidId = Secure.getString(activityContext.getApplicationContext().getContentResolver(), Secure.ANDROID_ID);
+			String androidId = Secure.getString(getApplicationContext().getContentResolver(), Secure.ANDROID_ID);
 			return androidId;
 		} catch (Exception e) {
 			Log.debug("GetAndroidId failed: " + e.getMessage());
@@ -4562,7 +4614,7 @@ public class GameActivity extends $${gameActivitySuperClass}$$ implements Surfac
 			if (bIsFile)
 			{
 				// Note that the ".fileprovider" suffix is hardcoded
-				String Authority = activityContext.getApplicationContext().getPackageName() + ".fileprovider";
+				String Authority = getApplicationContext().getPackageName() + ".fileprovider";
 				Uri FileUri = FileProvider.getUriForFile(this, Authority, SharedFile);
 				
 				SendIntent.putExtra(Intent.EXTRA_STREAM, FileUri);
@@ -4636,7 +4688,7 @@ public class GameActivity extends $${gameActivitySuperClass}$$ implements Surfac
 
 	public void AndroidThunkJava_SetCellularPreference(int value)
 	{
-		Context context = activityContext.getApplicationContext();
+		Context context = getApplicationContext();
 		SharedPreferences preferences = context.getSharedPreferences("CellularNetworkPreferences", context.MODE_PRIVATE);
 		SharedPreferences.Editor editor = preferences.edit();
 		editor.putInt("AllowCellular", value);
@@ -4645,7 +4697,7 @@ public class GameActivity extends $${gameActivitySuperClass}$$ implements Surfac
 
 	public int AndroidThunkJava_GetCellularPreference()
 	{
-		Context context = activityContext.getApplicationContext();
+		Context context = getApplicationContext();
 		SharedPreferences preferences = context.getSharedPreferences("CellularNetworkPreferences", context.MODE_PRIVATE);
 		return preferences.getInt("AllowCellular", 0 );
 	}
@@ -4691,7 +4743,7 @@ public class GameActivity extends $${gameActivitySuperClass}$$ implements Surfac
 
 	public boolean AndroidThunkJava_IsMusicActive()
 	{
-		AudioManager audioManager = (AudioManager)this.getSystemService(Context.AUDIO_SERVICE);
+		AudioManager audioManager = (AudioManager)getSystemService(Context.AUDIO_SERVICE);
 		return audioManager.isMusicActive();
 	}
 	
@@ -5045,7 +5097,7 @@ public class GameActivity extends $${gameActivitySuperClass}$$ implements Surfac
 
 	private ArrayList<Integer> LocalNotificationGetIDList()
 	{
-		SharedPreferences preferences = activityContext.getApplicationContext().getSharedPreferences("LocalNotificationPreferences", MODE_PRIVATE);
+		SharedPreferences preferences = getApplicationContext().getSharedPreferences("LocalNotificationPreferences", MODE_PRIVATE);
 		SharedPreferences.Editor editor = preferences.edit();
 		String notificationIDs = preferences.getString("notificationIDs", "");
 		ArrayList<Integer> iParts = new ArrayList<Integer>();
@@ -5293,7 +5345,7 @@ public class GameActivity extends $${gameActivitySuperClass}$$ implements Surfac
 		}
 
 		// clear them all
-		SharedPreferences preferences = activityContext.getApplicationContext().getSharedPreferences("LocalNotificationPreferences", MODE_PRIVATE);
+		SharedPreferences preferences = getApplicationContext().getSharedPreferences("LocalNotificationPreferences", MODE_PRIVATE);
 		SharedPreferences.Editor editor = preferences.edit();
 		editor.putString("notificationIDs", "");
 		editor.putString("notificationDetails", "{}");
@@ -5501,7 +5553,7 @@ public class GameActivity extends $${gameActivitySuperClass}$$ implements Surfac
 	{
 		boolean bPluginEnabled = false;
 //$${gameActivityAllowedRemoteNotificationsAdditions}$$
-		return bPluginEnabled && NotificationManagerCompat.from(activityContext.getApplicationContext()).areNotificationsEnabled();
+		return bPluginEnabled && NotificationManagerCompat.from(getApplicationContext()).areNotificationsEnabled();
 	}
 
 	public void AndroidThunkJava_SetDesiredViewSize(int width, int height)
@@ -5560,7 +5612,7 @@ public class GameActivity extends $${gameActivitySuperClass}$$ implements Surfac
 	{
 		if(ANDROID_BUILD_VERSION >= 24)
 		{
-			WindowManager windowManager = activityContext.getWindowManager();
+			WindowManager windowManager = getWindowManager();
 			Display display = windowManager.getDefaultDisplay();
 			Display.Mode currentmode = display.getMode();
 			Display.Mode[] modes = display.getSupportedModes();
@@ -5620,7 +5672,7 @@ public class GameActivity extends $${gameActivitySuperClass}$$ implements Surfac
 		}
 		else if(ANDROID_BUILD_VERSION >= 24)
 		{
-			WindowManager windowManager = activityContext.getWindowManager();
+			WindowManager windowManager = getWindowManager();
 			Display display = windowManager.getDefaultDisplay();
 			Display.Mode currentmode = display.getMode();
 			int currentmodeid = currentmode.getModeId();
@@ -5641,7 +5693,7 @@ public class GameActivity extends $${gameActivitySuperClass}$$ implements Surfac
 							{
 								public void run()
 								{
-									Window w = activityContext.getWindow();
+									Window w = getWindow();
 									WindowManager.LayoutParams l = w.getAttributes();
 									l.preferredDisplayModeId = modeid;
 									w.setAttributes(l);
@@ -5666,7 +5718,7 @@ public class GameActivity extends $${gameActivitySuperClass}$$ implements Surfac
 	{
 		if(ANDROID_BUILD_VERSION >= 24)
 		{
-			WindowManager windowManager = activityContext.getWindowManager();
+			WindowManager windowManager = getWindowManager();
 			Display display = windowManager.getDefaultDisplay();
 			Display.Mode currentmode = display.getMode();
 			return (int)currentmode.getRefreshRate();
@@ -5717,7 +5769,7 @@ public class GameActivity extends $${gameActivitySuperClass}$$ implements Surfac
 	{
 		assert(activityContext != null);
 
-		boolean bAirplaneMode = isAirplaneModeOn(activityContext.getApplicationContext());
+		boolean bAirplaneMode = isAirplaneModeOn(getApplicationContext());
 
 		ConnectivityManager cm = (ConnectivityManager)getSystemService(Context.CONNECTIVITY_SERVICE);
 		NetworkInfo activeNetwork = cm.getActiveNetworkInfo();
@@ -5742,7 +5794,7 @@ public class GameActivity extends $${gameActivitySuperClass}$$ implements Surfac
 	{
 		assert(activityContext != null);
 
-		boolean bAirplaneMode = isAirplaneModeOn(activityContext.getApplicationContext());
+		boolean bAirplaneMode = isAirplaneModeOn(getApplicationContext());
 
 		ConnectivityManager cm = (ConnectivityManager)getSystemService(Context.CONNECTIVITY_SERVICE);
 		NetworkInfo activeNetwork = cm.getActiveNetworkInfo();
@@ -5831,7 +5883,7 @@ public class GameActivity extends $${gameActivitySuperClass}$$ implements Surfac
 				}
 				else
 				{
-					return android.net.Proxy.getPort(activityContext.getApplicationContext());
+					return android.net.Proxy.getPort(getApplicationContext());
 				}
 			}
 			case "android.hardware.vulkan.version":
@@ -5886,7 +5938,7 @@ public class GameActivity extends $${gameActivitySuperClass}$$ implements Surfac
 		{
 			if (android.os.Build.VERSION.SDK_INT >= 21)
 			{
-				WindowManager windowManager = activityContext.getWindowManager();
+				WindowManager windowManager = getWindowManager();
 				android.view.Display display = windowManager.getDefaultDisplay();
 				return display.getPresentationDeadlineNanos();
 			}
@@ -5897,7 +5949,7 @@ public class GameActivity extends $${gameActivitySuperClass}$$ implements Surfac
 		{
 			if (android.os.Build.VERSION.SDK_INT >= 21)
 			{
-				WindowManager windowManager = activityContext.getWindowManager();
+				WindowManager windowManager = getWindowManager();
 				android.view.Display display = windowManager.getDefaultDisplay();
 				return display.getAppVsyncOffsetNanos();
 			}
@@ -5910,7 +5962,7 @@ public class GameActivity extends $${gameActivitySuperClass}$$ implements Surfac
 	{
 		if (key.equals("unreal.display.getRefreshRate"))
 		{
-			WindowManager windowManager = activityContext.getWindowManager();
+			WindowManager windowManager = getWindowManager();
 			android.view.Display display = windowManager.getDefaultDisplay();
 			float rate = display.getRefreshRate();
 			// make sure it is sane.. some drivers returned crazy values
@@ -5927,7 +5979,7 @@ public class GameActivity extends $${gameActivitySuperClass}$$ implements Surfac
 		}
 		else
 		{
-			return activityContext.getWindowManager().getDefaultDisplay();
+			return getWindowManager().getDefaultDisplay();
 		}
 	}
 
@@ -5936,17 +5988,17 @@ public class GameActivity extends $${gameActivitySuperClass}$$ implements Surfac
 		DisplayMetrics metrics = new DisplayMetrics();
 		if (android.os.Build.VERSION.SDK_INT >= 31)
 		{
-			metrics = activityContext.getResources().getDisplayMetrics();
+			metrics = getResources().getDisplayMetrics();
 		}
 		else
 		if (android.os.Build.VERSION.SDK_INT >= 17)
 		{
-			activityContext.getWindowManager().getDefaultDisplay().getRealMetrics(metrics);
+			getWindowManager().getDefaultDisplay().getRealMetrics(metrics);
 		}
 		else
 		{
 			// note: not available so get what we can
-			activityContext.getWindowManager().getDefaultDisplay().getMetrics(metrics);
+			getWindowManager().getDefaultDisplay().getMetrics(metrics);
 		}
 		return metrics;
 	}
@@ -6048,7 +6100,7 @@ public class GameActivity extends $${gameActivitySuperClass}$$ implements Surfac
 				}
 				else
 				{
-					return android.net.Proxy.getHost(activityContext.getApplicationContext());
+					return android.net.Proxy.getHost(getApplicationContext());
 				}
 			}
 			case "unreal.displaymetrics.dpi":
@@ -6151,7 +6203,7 @@ public class GameActivity extends $${gameActivitySuperClass}$$ implements Surfac
 					{
 						try
 						{
-							android.view.Window ActivityWindow = activityContext.getWindow();
+							android.view.Window ActivityWindow = getWindow();
 							Method m = android.view.Window.class.getMethod("setSustainedPerformanceMode",Boolean.TYPE);
 							m.invoke(ActivityWindow, bEnable);
 						}
@@ -6420,6 +6472,11 @@ public class GameActivity extends $${gameActivitySuperClass}$$ implements Surfac
 	//create the virtual keyboard input 
 	private void createVirtualKeyboardInput()
 	{
+		if (isStandalone)
+		{
+			return;
+		}
+
 		getWindow().setSoftInputMode(WindowManager.LayoutParams.SOFT_INPUT_ADJUST_RESIZE);
 		newVirtualKeyboardInput = new VirtualKeyboardInput(this);
 		newVirtualKeyboardInput.setSingleLine(false);
@@ -6611,7 +6668,7 @@ public class GameActivity extends $${gameActivitySuperClass}$$ implements Surfac
 
 		if (intentAction.equals(""))
 		{
-			intent = activityContext.getPackageManager().getLaunchIntentForPackage(packageName);
+			intent = getPackageManager().getLaunchIntentForPackage(packageName);
 			if (intent == null)
 			{
 				return false;
@@ -6638,7 +6695,7 @@ public class GameActivity extends $${gameActivitySuperClass}$$ implements Surfac
 			intent.setComponent(new android.content.ComponentName(packageName, intentComponent));
 		}
 
-		if (intent.resolveActivity(activityContext.getPackageManager()) != null)
+		if (intent.resolveActivity(getPackageManager()) != null)
 		{
 			startActivity(intent);
 			AndroidThunkJava_ForceQuit();
@@ -6675,7 +6732,7 @@ public class GameActivity extends $${gameActivitySuperClass}$$ implements Surfac
 
 		if (intentAction.equals(""))
 		{
-			intent = activityContext.getPackageManager().getLaunchIntentForPackage(packageName);
+			intent = getPackageManager().getLaunchIntentForPackage(packageName);
 			if (intent == null)
 			{
 				return false;
@@ -6702,7 +6759,7 @@ public class GameActivity extends $${gameActivitySuperClass}$$ implements Surfac
 			intent.setComponent(new android.content.ComponentName(packageName, intentComponent));
 		}
 
-		activityContext.getApplicationContext().sendBroadcast(intent);
+		getApplicationContext().sendBroadcast(intent);
 
 		if (bExit)
 		{
@@ -6714,7 +6771,7 @@ public class GameActivity extends $${gameActivitySuperClass}$$ implements Surfac
 
 	protected void RestartApplication(String RestartExtra)
 	{
-		Context context = activityContext.getApplicationContext();
+		Context context = getApplicationContext();
 	
 		// get the launch intent
 		PackageManager packageManager = context.getPackageManager();
@@ -6741,7 +6798,7 @@ public class GameActivity extends $${gameActivitySuperClass}$$ implements Surfac
 		};
 		Log.debug("app restart : " + RestartExtra);
 //$${gameActivityOnRestartApplicationAdditions}$$
-		runOnUiThread(restartRunnableUIThread);
+		activityContext.runOnUiThread(restartRunnableUIThread);
 
 		// wait for exit.
 		while(true)
@@ -6781,13 +6838,13 @@ public class GameActivity extends $${gameActivitySuperClass}$$ implements Surfac
 
 	public int AndroidThunkJava_GetDeviceRotation() 
 	{
-		return activityContext.getWindowManager().getDefaultDisplay().getRotation();
+		return getWindowManager().getDefaultDisplay().getRotation();
 	}
 
 	public int AndroidThunkJava_GetDeviceOrientation()
 	{
 		int currentScreenOrientation = DEVICE_ORIENTATION_UNKNOWN;
-		int rotation = activityContext.getWindowManager().getDefaultDisplay().getRotation();
+		int rotation = getWindowManager().getDefaultDisplay().getRotation();
 		if ( getDeviceDefaultOrientation() == Configuration.ORIENTATION_LANDSCAPE )
 		{
 			switch (rotation)
@@ -6835,7 +6892,7 @@ public class GameActivity extends $${gameActivitySuperClass}$$ implements Surfac
 		}
 		else
 		{
-			Context context = activityContext.getApplicationContext();
+			Context context = getApplicationContext();
 			CookieSyncManager syncManager = CookieSyncManager.createInstance(context);
 			syncManager.sync();
 
@@ -6858,7 +6915,7 @@ public class GameActivity extends $${gameActivitySuperClass}$$ implements Surfac
 	{
 		if (android.os.Build.VERSION.SDK_INT < android.os.Build.VERSION_CODES.LOLLIPOP)
 		{
-			Context context = activityContext.getApplicationContext();
+			Context context = getApplicationContext();
 			CookieSyncManager syncManager = CookieSyncManager.createInstance(context);
 			syncManager.sync();
 		}
@@ -6877,7 +6934,7 @@ public class GameActivity extends $${gameActivitySuperClass}$$ implements Surfac
 	{
 		if (android.os.Build.VERSION.SDK_INT < android.os.Build.VERSION_CODES.LOLLIPOP)
 		{
-			Context context = activityContext.getApplicationContext();
+			Context context = getApplicationContext();
 			CookieSyncManager syncManager = CookieSyncManager.createInstance(context);
 			syncManager.sync();
 
diff --git a/Engine/Plugins/Runtime/AndroidSingleInstanceService/Source/AndroidSingleInstanceService/Android/ASISStub/src/main/java/com/epicgames/ASISStub/ASISConnection.java b/Engine/Plugins/Runtime/AndroidSingleInstanceService/Source/AndroidSingleInstanceService/Android/ASISStub/src/main/java/com/epicgames/ASISStub/ASISConnection.java
new file mode 100644
--- /dev/null
+++ b/Engine/Plugins/Runtime/AndroidSingleInstanceService/Source/AndroidSingleInstanceService/Android/ASISStub/src/main/java/com/epicgames/ASISStub/ASISConnection.java
@@ -0,0 +1,445 @@
+package com.epicgames.ASISStub;
+
+import android.app.Activity;
+import android.app.Application;
+import android.content.ComponentName;
+import android.content.Context;
+import android.content.ContextWrapper;
+import android.content.Intent;
+import android.content.ServiceConnection;
+import android.graphics.SurfaceTexture;
+import android.os.Build;
+import android.os.Bundle;
+import android.os.Handler;
+import android.os.IBinder;
+import android.os.Message;
+import android.os.Messenger;
+import android.os.RemoteException;
+import android.util.Log;
+import android.view.MotionEvent;
+import android.view.Surface;
+import android.view.TextureView;
+import android.view.View;
+
+import androidx.annotation.NonNull;
+import androidx.annotation.UiThread;
+import androidx.lifecycle.MutableLiveData;
+
+import com.epicgames.makeaar.UnrealMessageType;
+
+import java.util.concurrent.atomic.AtomicReference;
+
+
+public class ASISConnection implements ServiceConnection, TextureView.SurfaceTextureListener, TextureView.OnTouchListener
+{
+	final static String TAG = "UE-ASISConnection";
+
+	final static String obbModuleName = "Lyra";//""ALPHA_POC";
+	//	final static String obbModuleName = BuildConfig.ASISModuleName; //"Lyra";//""ALPHA_POC";
+//	final static String servicePackageName = "";
+//	final static String servicePackageName = BuildConfig.ASISPackageName;
+	final static String servicePackageName = "com.epicgames.Lyra";
+	final static String obbFileLocation = ""; //"/data/local/tmp/3d/AndroidPackagingTest.main.obb.png";
+//	final static String obbFileLocation = BuildConfig.ASISOBBPath; //"/data/local/tmp/3d/AndroidPackagingTest.main.obb.png";
+
+	final static String commandLineArgs = "-nosound";
+	//	final static String commandLineArgs = "-nosound -launchandroidflags=0 -dpcvars='Android.UseGameThreadForNativeCommands=0'";
+	final static boolean enablePropagateAlpha = true;
+	
+	Messenger serviceMessenger = null;
+	Messenger serviceReplyMessenger = null;
+	MutableLiveData<Boolean> isBoundToService = new MutableLiveData<Boolean>(false);
+
+
+	AtomicReference<Surface> externalSurfaceData = new AtomicReference<Surface>(null);
+
+	int[] surfacePosition = new int[2];
+
+	TextureView activeTextureView = null;
+
+	final int mTaskID;
+
+	private static final Handler.Callback mConnectionStateCallback = (Message msg) -> {
+		final String TAG = "UE-ServiceCallback";
+
+		switch (msg.what) {
+
+			case 0:	//EVENTTYPE_INIT == 0
+				Log.v(TAG, "received mConnectionStateCallback for EVENTTYPE_INIT, obj = " + msg.obj);
+
+				break;
+
+			case 1:	//EVENTTYPE_POST_ENGINE_INIT == 1
+				Log.v(TAG, "received mConnectionStateCallback for EVENTTYPE_POST_ENGINE_INIT, obj = " + msg.obj);
+				break;
+
+			default:
+				//Log.e(TAG, "mConnectionStateCallback, Unknown message = " + msg.what);
+		}
+		return true;
+	};
+	
+
+	int getTaskId() { return mTaskID; }
+
+	ContextWrapper mContextWrapper;
+
+	ASISConnection(Activity _activity)
+	{
+//			if (mServiceConnection.getValue() != null)
+//			{
+//				unbindToUnrealInstanceService("MyServiceConnection constructor!!!");
+//				doCleanupForUnbinding("MyServiceConnection constructor!!!");
+//
+//			}
+		assert (_activity != null);
+
+		mContextWrapper = _activity;
+
+		mTaskID = _activity.getTaskId();
+
+//			bindToUnrealInstanceService("MyServiceConnection constructor " + mTaskID);
+	}
+
+
+//		Activity getActivity() {
+//			assert (mContextWrapper instanceof Activity);
+//			return (Activity) mContextWrapper;
+//		}
+//
+//		Application getApplication()
+//		{
+//			Application application = null;
+//			try {
+//				application = getActivity().getApplication();
+//			}
+//			catch (Exception e)
+//			{
+//				e.printStackTrace();
+//			}
+//			finally {
+//				return application;
+//			}
+//		}
+
+	void SendMessage(Message msg) throws RemoteException {
+		if (serviceMessenger != null) {
+			serviceMessenger.send(msg);
+		}
+	}
+
+	void SetTextureView(TextureView textureView_)
+	{
+		textureView_.setAlpha(1.0f);
+		textureView_.setOpaque(false);
+		textureView_.setOnTouchListener(this);
+		textureView_.setSurfaceTextureListener(this);
+
+		Log.d(TAG, "SetTextureView: " + textureView_ + ", isAvailable=" + textureView_.isAvailable() + ", isAttachedToWindow=" + textureView_.isAttachedToWindow());
+
+		if (textureView_.isAvailable() && textureView_.isAttachedToWindow())
+		{
+			textureView_.getLocationOnScreen(surfacePosition);
+
+			onSurfaceTextureAvailable(textureView_.getSurfaceTexture(), textureView_.getWidth(), textureView_.getHeight());
+		}
+		activeTextureView = textureView_;
+	}
+
+	public boolean onTouch(View v, MotionEvent event) {
+		// TODO Auto-generated method stub
+
+		Log.d(TAG, "Touch: " + event.getAction() + " X: " + event.getX() + " Y: " + event.getY());
+
+		Message msg = Message.obtain(null, UnrealMessageType.TouchEvent.ordinal());
+		msg.getData().putInt("taskId", getTaskId());
+		msg.getData().putParcelable("touch", event);
+
+		try {
+			SendMessage(msg);
+		} catch (RemoteException e) {
+			e.printStackTrace();
+		}
+
+		return true;
+	}
+
+	public void ConsoleCommand(Context context, String consoleCommand)
+	{
+		Log.i(TAG, "ConsoleCommand() - consoleCommand = " + consoleCommand);
+
+		// Create and send a message to the service, using a supported 'what' value
+		Message msg = Message.obtain(null, UnrealMessageType.SendConsoleCommand.ordinal());
+		msg.getData().putString("consoleCommand", consoleCommand);
+
+		try {
+			SendMessage(msg);
+		} catch (RemoteException e) {
+			e.printStackTrace();
+		}
+	}
+
+	void SendData(int event, String param1, int param2, int param3, float param4)
+	{
+		// Create and send a message to the service, using a supported 'what' value
+		Message msg = Message.obtain(null, UnrealMessageType.SendData.ordinal());
+		msg.getData().putInt("taskId", getTaskId());
+		msg.getData().putInt("event", event);
+		msg.getData().putString("param1", param1);
+		msg.getData().putInt("param2", param2);
+		msg.getData().putInt("param3", param3);
+		msg.getData().putFloat("param4", param4);
+
+		try {
+			SendMessage(msg);
+		} catch (RemoteException e) {
+			e.printStackTrace();
+		}
+	}
+
+	void bindToUnrealInstanceService(String caller)
+	{
+		logContextDetails("bindToUnrealInstanceService("+ caller + ") -> ");
+		if(serviceReplyMessenger == null)
+		{
+			serviceReplyMessenger = new Messenger(new Handler(mConnectionStateCallback));
+		}
+
+		Intent intent = new Intent();
+		intent.setClassName(servicePackageName, "com.epicgames.makeaar.UnrealSharedInstanceService");
+
+		intent.putExtra("obbModuleName", obbModuleName);
+		intent.putExtra("obbFileLocation", obbFileLocation);
+		intent.putExtra("commandLineArgs", commandLineArgs);
+		intent.putExtra("enablePropagateAlpha", enablePropagateAlpha);
+
+		if (serviceReplyMessenger != null) {
+			intent.putExtra("callback", serviceReplyMessenger);
+		}
+
+		//ServiceConnection connection = GetServiceConnection("bindToUnrealInstanceService, caller=" + caller);
+
+		Log.d(TAG, "bindToUnrealInstanceService intent=" + intent + ", mServiceConnection=" + this);
+		mContextWrapper.bindService(intent, this, Context.BIND_AUTO_CREATE | Context.BIND_IMPORTANT | Context.BIND_DEBUG_UNBIND | Context.BIND_ADJUST_WITH_ACTIVITY);
+//			getApplication().bindService(intent, this, Context.BIND_AUTO_CREATE | Context.BIND_IMPORTANT | Context.BIND_ADJUST_WITH_ACTIVITY);
+	}
+
+	void unbindToUnrealInstanceService(String caller)
+	{
+		logContextDetails("unbindToUnrealInstanceService("+ caller + ") -> ");
+
+		if (isBoundToService.getValue()) {
+			Log.v(TAG, "unbindToUnrealInstanceService");
+
+			mContextWrapper.unbindService(this);
+		}
+	}
+
+
+	private void logContextDetails(String caller)
+	{
+		Log.w(TAG, caller + ": isBoundToService=" +isBoundToService.getValue()
+			+ ", taskID=" +getTaskId()
+			+ ", serviceMessenger=" +serviceMessenger
+			+ ", serviceReplyMessenger=" +serviceReplyMessenger
+			+ ", externalSurfaceData=" +externalSurfaceData
+		);
+	}
+
+	public void doCleanupForUnbinding(String caller) {
+		logContextDetails("doCleanupForUnbinding(" + caller + ")");
+		
+		if (activeTextureView != null) {
+			activeTextureView.setSurfaceTextureListener(null);
+			activeTextureView = null;
+		}
+
+		if (externalSurfaceData.get() != null)
+		{
+			externalSurfaceData.get().release();
+			externalSurfaceData.set(null);
+		}
+
+		serviceMessenger = null;
+		serviceReplyMessenger = null;
+		isBoundToService.setValue( false );
+	}
+
+	@Override
+	public void onServiceConnected(ComponentName name, IBinder serviceBinder) {
+		Log.v(TAG, "onServiceConnected name=" + name + ", serviceBinder=" + serviceBinder);
+		serviceMessenger = new Messenger(serviceBinder);
+		//serviceReplyMessenger = new Messenger(new Handler(mConnectionStateCallback));
+		isBoundToService.setValue( true );
+//			TextureView textureView_ = (TextureView) findViewById( R.id.renderView );
+//
+//			textureView_ = (TextureView) getActivity().findViewById( R.id.renderView );
+//			SetTextureView(textureView_);
+	}
+
+	@Override
+	public void onServiceDisconnected(ComponentName name) {
+
+		Log.w(TAG, "onServiceDisconnected name=" + name);
+
+		doCleanupForUnbinding("onServiceDisconnected name=" + name);
+	}
+
+	@Override
+	public void onBindingDied(ComponentName name) {
+
+		Log.w(TAG, "onBindingDied name=" + name);
+
+		doCleanupForUnbinding("onBindingDied name=" + name);
+
+		ServiceConnection.super.onBindingDied(name);
+	}
+
+	@Override
+	public void onNullBinding(ComponentName name) {
+
+		Log.w(TAG, "onNullBinding name=" + name);
+
+		//doCleanupForBinding("onNullBinding");
+
+		ServiceConnection.super.onNullBinding(name);
+	}
+
+	@Override
+	public void onSurfaceTextureAvailable(@NonNull SurfaceTexture surfaceTexture, int width, int height) {
+
+		logContextDetails("onSurfaceTextureAvailable(" + surfaceTexture + ")");
+
+		int[] surfacePosition = this.surfacePosition;
+		int[] surfaceSize = new int[]{width, height};
+
+		if (externalSurfaceData.get() != null) {
+			externalSurfaceData.get().release();
+		}
+
+		externalSurfaceData.set(new Surface(surfaceTexture));
+
+		attachSurfaceToService(getTaskId(), externalSurfaceData.get(), surfacePosition, surfaceSize);
+	}
+
+	@Override
+	public void onSurfaceTextureSizeChanged(@NonNull SurfaceTexture surfaceTexture, int width, int height) {
+		logContextDetails("onSurfaceTextureSizeChanged(" + surfaceTexture + ")");
+	}
+
+	// Invoked when the specified SurfaceTexture is about to be destroyed.
+	// If returns true, no rendering should happen inside the surface texture after this method is invoked.
+	// If returns false, the client needs to call SurfaceTexture.release().
+	// Most applications should return true.
+	@Override
+	public boolean onSurfaceTextureDestroyed(@NonNull SurfaceTexture surfaceTexture) {
+
+		synchronized (this) {
+			logContextDetails("onSurfaceTextureDestroyed(" + surfaceTexture + ")");
+			detachSurfaceFromService(getTaskId(), null);
+
+			if (externalSurfaceData.get() != null) {
+				externalSurfaceData.get().release();
+				externalSurfaceData.set(null);
+			}
+
+			if (activeTextureView != null) {
+				activeTextureView.setSurfaceTextureListener(null);
+			}
+			return true;
+		}
+	}
+
+	@Override
+	public void onSurfaceTextureUpdated(@NonNull SurfaceTexture surfaceTexture) {
+	}
+
+	@UiThread
+	void attachSurfaceToService(int attachId, Surface externalSurface, int[] viewPosition, int[] viewSize)
+	{
+		synchronized (this) {
+			if (externalSurface == null) {
+				Log.w(TAG, "attachSurfaceToService, externalSurface is null");
+				return;
+			}
+
+			if (Build.VERSION.SDK_INT >= Build.VERSION_CODES.P) {
+				Log.v(TAG, "attachSurfaceToService() - " + "proc = " + Application.getProcessName() + ", externalSurface = " + externalSurface + ", attachId = " + attachId);
+			}
+
+			Message msg = Message.obtain(null, UnrealMessageType.AttachExternalSurface.ordinal());
+			Bundle data = msg.getData();
+			data.putInt("attachId", attachId);
+			data.putBoolean("enablePropagateAlpha", true);
+			data.putBoolean("start", true);
+			data.putBoolean("resume", true);
+
+			data.putParcelable("surface", externalSurface);
+
+			if (viewPosition.length == 2) {
+				data.putIntArray("viewPosition", viewPosition);
+			}
+
+			if (viewSize.length == 2) {
+				data.putIntArray("viewSize", viewSize);
+
+			}
+
+			try {
+				SendMessage(msg);
+			} catch (RemoteException e) {
+				e.printStackTrace();
+			}
+		}
+	}
+
+	@UiThread
+	void detachSurfaceFromService(int attachId, Handler.Callback callback) 
+	{
+
+		synchronized (this) {
+			Message msg = Message.obtain(
+				null,
+				UnrealMessageType.DetachExternalSurface.ordinal()
+			);
+
+			msg.getData().putInt("attachId", attachId);
+			msg.getData().putBoolean("pause", true);
+			msg.getData().putBoolean("stop", true);
+
+			if (callback != null) {
+
+				Handler handler = new Handler(callback);
+				Messenger replyMessenger = new Messenger(handler);
+
+				msg.replyTo = replyMessenger;
+			}
+
+			Log.i(TAG, "detachSurfaceFromService() - " + "proc = ${Application.getProcessName()}");
+
+			try {
+				SendMessage(msg);
+			} catch (RemoteException e) {
+				e.printStackTrace();
+			}
+		}
+	}
+
+	@UiThread
+	void resumeService(int attachId)
+	{
+		synchronized (this) {
+			if (isBoundToService.getValue()) {
+				Message msg = Message.obtain(null, UnrealMessageType.ResumeService.ordinal());
+				Bundle data = msg.getData();
+				data.putInt("attachId", attachId);
+
+				try {
+					SendMessage(msg);
+				} catch (RemoteException e) {
+					e.printStackTrace();
+				}
+			}
+		}
+	}
+}
\ No newline at end of file
diff --git a/Engine/Plugins/Runtime/AndroidSingleInstanceService/Source/AndroidSingleInstanceService/Android/ASISStub/src/main/java/com/epicgames/ASISStub/UseServiceActivity.java b/Engine/Plugins/Runtime/AndroidSingleInstanceService/Source/AndroidSingleInstanceService/Android/ASISStub/src/main/java/com/epicgames/ASISStub/UseServiceActivity.java
--- a/Engine/Plugins/Runtime/AndroidSingleInstanceService/Source/AndroidSingleInstanceService/Android/ASISStub/src/main/java/com/epicgames/ASISStub/UseServiceActivity.java
+++ b/Engine/Plugins/Runtime/AndroidSingleInstanceService/Source/AndroidSingleInstanceService/Android/ASISStub/src/main/java/com/epicgames/ASISStub/UseServiceActivity.java
@@ -28,7 +28,10 @@ import android.view.Surface;
 import android.view.TextureView;
 import android.view.View;
 import android.content.Intent;
+import android.view.Window;
+import android.view.WindowManager;
 import android.widget.Button;
+import android.widget.RelativeLayout;
 
 import com.epicgames.makeaar.UnrealMessageType;
 
@@ -42,439 +45,74 @@ public class UseServiceActivity extends Activity implements View.OnClickListener
 	Messenger mMessenger;
 //	private static WeakReference<View.OnClickListener> mCallbackRef = new WeakReference<>(IEventCallback.NO_OP);
 
-	private static final Handler.Callback mConnectionStateCallback = (Message msg) -> {
-		final String TAG = "UE-ServiceCallback";
 
-		switch (msg.what) {
-
-			case 0:	//EVENTTYPE_INIT == 0
-				Log.v(TAG, "received mConnectionStateCallback for EVENTTYPE_INIT, obj = " + msg.obj);
-
-				break;
-
-			case 1:	//EVENTTYPE_POST_ENGINE_INIT == 1
-				Log.v(TAG, "received mConnectionStateCallback for EVENTTYPE_POST_ENGINE_INIT, obj = " + msg.obj);
-				break;
-
-			default:
-				//Log.e(TAG, "mConnectionStateCallback, Unknown message = " + msg.what);
-		}
-		return true;
-	};
 
 	//final MutableLiveData<MyServiceConnection> mServiceConnection = new MutableLiveData<MyServiceConnection>(new MyServiceConnection());
 
-	static MyServiceConnection mServiceConnection = null; //new MyServiceConnection();
-
-	static class MyServiceConnection implements ServiceConnection, TextureView.SurfaceTextureListener, TextureView.OnTouchListener
+	ASISConnection mServiceConnection = null; //new MyServiceConnection();
+	
+	ASISConnection GetServiceConnection(String caller)
 	{
-		Messenger serviceMessenger = null;
-		Messenger serviceReplyMessenger = null;
-		static MutableLiveData<Boolean> isBoundToService = new MutableLiveData<Boolean>(false);
-
-
-		AtomicReference<Surface> externalSurfaceData = new AtomicReference<Surface>(null);
-
-		int[] surfacePosition = new int[2];
-
-		TextureView activeTextureView = null;
-
-		final int mTaskID;
-
-		int getTaskId() { return mTaskID; }
-
-		ContextWrapper mContextWrapper;
-
-		MyServiceConnection(Activity _activity)
-		{
-//			if (mServiceConnection.getValue() != null)
-//			{
-//				unbindToUnrealInstanceService("MyServiceConnection constructor!!!");
-//				doCleanupForUnbinding("MyServiceConnection constructor!!!");
-//
-//			}
-			assert (_activity != null);
-
-			mContextWrapper = _activity;
-
-			mTaskID = _activity.getTaskId();
-
-//			bindToUnrealInstanceService("MyServiceConnection constructor " + mTaskID);
-		}
-
-
-//		Activity getActivity() {
-//			assert (mContextWrapper instanceof Activity);
-//			return (Activity) mContextWrapper;
-//		}
-//
-//		Application getApplication()
-//		{
-//			Application application = null;
-//			try {
-//				application = getActivity().getApplication();
-//			}
-//			catch (Exception e)
-//			{
-//				e.printStackTrace();
-//			}
-//			finally {
-//				return application;
-//			}
-//		}
-
-		void SendMessage(Message msg) throws RemoteException {
-			if (serviceMessenger != null) {
-				serviceMessenger.send(msg);
-			}
-		}
-
-		void SetTextureView(TextureView textureView_)
-		{
-			textureView_.setAlpha(1.0f);
-			textureView_.setOpaque(false);
-			//textureView_.setOnTouchListener(this);
-			textureView_.setSurfaceTextureListener(this);
-
-			Log.d(TAG, "SetTextureView: " + textureView_ + ", isAvailable=" + textureView_.isAvailable() + ", isAttachedToWindow=" + textureView_.isAttachedToWindow());
-
-			if (textureView_.isAvailable() && textureView_.isAttachedToWindow())
-			{
-				textureView_.getLocationOnScreen(surfacePosition);
-
-				onSurfaceTextureAvailable(textureView_.getSurfaceTexture(), textureView_.getWidth(), textureView_.getHeight());
-			}
-			activeTextureView = textureView_;
-		}
-
-		public boolean onTouch(View v, MotionEvent event) {
-			// TODO Auto-generated method stub
-
-			Log.d(TAG, "Touch: " + event.getAction() + " X: " + event.getX() + " Y: " + event.getY());
-
-			Message msg = Message.obtain(null, UnrealMessageType.TouchEvent.ordinal());
-			msg.getData().putInt("taskId", getTaskId());
-			msg.getData().putParcelable("touch", event);
 
-			try {
-				SendMessage(msg);
-			} catch (RemoteException e) {
-				e.printStackTrace();
-			}
-
-			return true;
-		}
-
-		public void ConsoleCommand(Context context, String consoleCommand)
-		{
-			Log.i(TAG, "ConsoleCommand() - consoleCommand = " + consoleCommand);
-
-			// Create and send a message to the service, using a supported 'what' value
-			Message msg = Message.obtain(null, UnrealMessageType.SendConsoleCommand.ordinal());
-			msg.getData().putString("consoleCommand", consoleCommand);
-
-			try {
-				SendMessage(msg);
-			} catch (RemoteException e) {
-				e.printStackTrace();
-			}
-		}
-
-		void SendData(int event, String param1, int param2, int param3, float param4)
-		{
-			// Create and send a message to the service, using a supported 'what' value
-			Message msg = Message.obtain(null, UnrealMessageType.SendData.ordinal());
-			msg.getData().putInt("taskId", getTaskId());
-			msg.getData().putInt("event", event);
-			msg.getData().putString("param1", param1);
-			msg.getData().putInt("param2", param2);
-			msg.getData().putInt("param3", param3);
-			msg.getData().putFloat("param4", param4);
-
-			try {
-				SendMessage(msg);
-			} catch (RemoteException e) {
-				e.printStackTrace();
-			}
-		}
+		Log.w(TAG, "GetServiceConnection(" + caller + ") : mServiceConnection = " + mServiceConnection);
 
-		void bindToUnrealInstanceService(String caller)
+		if (mServiceConnection != null)
 		{
-			logContextDetails("bindToUnrealInstanceService("+ caller + ") -> ");
-			if(serviceReplyMessenger == null)
+			if (!mServiceConnection.isBoundToService.getValue())
 			{
-				serviceReplyMessenger = new Messenger(new Handler(mConnectionStateCallback));
-			}
-
-			Intent intent = new Intent();
-			intent.setClassName(servicePackageName, "com.epicgames.makeaar.UnrealSharedInstanceService");
-
-			intent.putExtra("obbModuleName", obbModuleName);
-			intent.putExtra("obbFileLocation", obbFileLocation);
-			intent.putExtra("commandLineArgs", commandLineArgs);
-			intent.putExtra("enablePropagateAlpha", enablePropagateAlpha);
-
-			if (serviceReplyMessenger != null) {
-				intent.putExtra("callback", serviceReplyMessenger);
-			}
-
-			Log.d(TAG, "bindToUnrealInstanceService intent=" + intent);
-			mContextWrapper.bindService(intent, this,  Context.BIND_AUTO_CREATE | Context.BIND_IMPORTANT | Context.BIND_DEBUG_UNBIND | Context.BIND_ADJUST_WITH_ACTIVITY);
-//			getApplication().bindService(intent, this, Context.BIND_AUTO_CREATE | Context.BIND_IMPORTANT | Context.BIND_ADJUST_WITH_ACTIVITY);
-		}
-
-		void unbindToUnrealInstanceService(String caller)
-		{
-			logContextDetails("unbindToUnrealInstanceService("+ caller + ") -> ");
-
-			if (isBoundToService.getValue()) {
-				Log.v(TAG, "unbindToUnrealInstanceService");
-
-				mContextWrapper.unbindService(this);
+				mServiceConnection.bindToUnrealInstanceService("GetServiceConnection rebind");
 			}
+			return mServiceConnection;
 		}
-
-
-		private void logContextDetails(String caller)
+		else
 		{
-			Log.w(TAG, caller + ": isBoundToService=" +isBoundToService.getValue()
-					+ ", taskID=" +getTaskId()
-					+ ", serviceMessenger=" +serviceMessenger
-					+ ", serviceReplyMessenger=" +serviceReplyMessenger
-					+ ", externalSurfaceData=" +externalSurfaceData
-			);
+			CreateServiceConnection(this, "GetServiceConnection");
 		}
-
-		private void doCleanupForUnbinding(String caller) {
-			logContextDetails("doCleanupForBinding(" + caller + ")");
-
-//			if (isBoundToService.getValue())
-//			{
-//				//getApplication().unbindService(this);
-//				//mServiceConnection.setValue( null );
-//			}
-
-			if (activeTextureView != null) {
-				activeTextureView.setSurfaceTextureListener(null);
-				activeTextureView = null;
-			}
-
-			if (externalSurfaceData.get() != null)
-			{
-				externalSurfaceData.get().release();
-				externalSurfaceData.set(null);
-			}
-
-			serviceMessenger = null;
-			serviceReplyMessenger = null;
-			isBoundToService.setValue( false );
-		}
-
-		@Override
-		public void onServiceConnected(ComponentName name, IBinder serviceBinder) {
-			Log.v(TAG, "onServiceConnected name=" + name + ", serviceBinder=" + serviceBinder);
-			serviceMessenger = new Messenger(serviceBinder);
-			//serviceReplyMessenger = new Messenger(new Handler(mConnectionStateCallback));
-			isBoundToService.setValue( true );
-//			TextureView textureView_ = (TextureView) findViewById( R.id.renderView );
-//
-//			textureView_ = (TextureView) getActivity().findViewById( R.id.renderView );
-//			SetTextureView(textureView_);
-		}
-
-		@Override
-		public void onServiceDisconnected(ComponentName name) {
-
-			Log.w(TAG, "onServiceDisconnected name=" + name);
-
-			doCleanupForUnbinding("onServiceDisconnected name=" + name);
-		}
-
-		@Override
-		public void onBindingDied(ComponentName name) {
-
-			Log.w(TAG, "onBindingDied name=" + name);
-
-			doCleanupForUnbinding("onBindingDied name=" + name);
-
-			ServiceConnection.super.onBindingDied(name);
-		}
-
-		@Override
-		public void onNullBinding(ComponentName name) {
-
-			Log.w(TAG, "onNullBinding name=" + name);
-
-			//doCleanupForBinding("onNullBinding");
-
-			ServiceConnection.super.onNullBinding(name);
-		}
-
-		@Override
-		public void onSurfaceTextureAvailable(@NonNull SurfaceTexture surfaceTexture, int width, int height) {
-
-			logContextDetails("onSurfaceTextureAvailable(" + surfaceTexture + ")");
-
-			int[] surfacePosition = this.surfacePosition;
-			int[] surfaceSize = new int[]{width, height};
-
-			if (externalSurfaceData.get() != null) {
-				externalSurfaceData.get().release();
-			}
-
-			externalSurfaceData.set(new Surface(surfaceTexture));
-
-			attachSurfaceToService(getTaskId(), externalSurfaceData.get(), surfacePosition, surfaceSize);
-		}
-
-		@Override
-		public void onSurfaceTextureSizeChanged(@NonNull SurfaceTexture surfaceTexture, int width, int height) {
-			logContextDetails("onSurfaceTextureSizeChanged(" + surfaceTexture + ")");
-		}
-
-		// Invoked when the specified SurfaceTexture is about to be destroyed.
-		// If returns true, no rendering should happen inside the surface texture after this method is invoked.
-		// If returns false, the client needs to call SurfaceTexture.release().
-		// Most applications should return true.
-		@Override
-		public boolean onSurfaceTextureDestroyed(@NonNull SurfaceTexture surfaceTexture) {
-
-			logContextDetails("onSurfaceTextureDestroyed(" + surfaceTexture + ")");
-			detachSurfaceFromService(getTaskId(), null);
-
-			if (externalSurfaceData.get() != null) {
-				externalSurfaceData.get().release();
-				externalSurfaceData.set(null);
-			}
-
-			if (activeTextureView != null) {
-				activeTextureView.setSurfaceTextureListener(null);
-			}
-			return true;
-		}
-
-		@Override
-		public void onSurfaceTextureUpdated(@NonNull SurfaceTexture surfaceTexture) {
-		}
-
-		@UiThread
-		void attachSurfaceToService(int attachId, Surface externalSurface, int[] viewPosition, int[] viewSize)
-		{
-			if (externalSurface ==  null) {
-				Log.w(TAG, "attachSurfaceToService, externalSurface is null");
-				return;
-			}
-
-			if (Build.VERSION.SDK_INT >= Build.VERSION_CODES.P) {
-				Log.v(TAG, "attachSurfaceToService() - " + "proc = " + Application.getProcessName() + ", externalSurface = " + externalSurface + ", attachId = " + attachId);
-			}
-
-			Message msg = Message.obtain(null, UnrealMessageType.AttachExternalSurface.ordinal());
-			Bundle data = msg.getData();
-			data.putInt("attachId", attachId);
-			data.putBoolean("enablePropagateAlpha", true);
-
-			data.putParcelable("surface", externalSurface);
-
-			if (viewPosition.length == 2) {
-				data.putIntArray("viewPosition", viewPosition);
-			}
-
-			if (viewSize.length == 2) {
-				data.putIntArray("viewSize", viewSize);
-
-			}
-
-			try {
-				SendMessage(msg);
-			} catch (RemoteException e) {
-				e.printStackTrace();
-			}
-		}
-
-		@UiThread
-		void detachSurfaceFromService(int attachId, Handler.Callback callback) {
-
-			Message msg = Message.obtain(
-					null,
-					UnrealMessageType.DetachExternalSurface.ordinal()
-			);
-
-			msg.getData().putInt("attachId", attachId);
-
-			if (callback != null) {
-
-				Handler handler = new Handler(callback);
-				Messenger replyMessenger = new Messenger(handler);
-
-				msg.replyTo = replyMessenger;
-			}
-
-			Log.i(TAG, "detachSurfaceFromService() - " + "proc = ${Application.getProcessName()}");
-
-			try {
-				SendMessage(msg);
-			} catch (RemoteException e) {
-				e.printStackTrace();
-			}
-		}
-
-		@UiThread
-		void resumeService(int attachId)
-		{
-			if (isBoundToService.getValue()) {
-				Message msg = Message.obtain(null, UnrealMessageType.ResumeService.ordinal());
-				Bundle data = msg.getData();
-				data.putInt("attachId", attachId);
-
-				try {
-					SendMessage(msg);
-				} catch (RemoteException e) {
-					e.printStackTrace();
-				}
-			}
-		}
-	}
-
-	MyServiceConnection GetServiceConnection(String caller)
-	{
-
-		Log.w(TAG, "GetServiceConnection(" + caller + ") : mServiceConnection = " + mServiceConnection);
-
+		
 		assert (mServiceConnection != null);
 
 		return mServiceConnection;
 	}
 
-	MyServiceConnection CreateServiceConnection(Activity activity, String caller)
+	ASISConnection CreateServiceConnection(Activity activity, String caller)
 	{
 
 		Log.w(TAG, "CreateServiceConnection(" + caller + ") : mServiceConnection = " + mServiceConnection + ", activity=" + activity);
 
 		if (mServiceConnection != null)
 		{
-			mServiceConnection.logContextDetails("CreateServiceConnection destroy prev (" + activity + ")");
-			DeposeServiceConnection();
+			if (!mServiceConnection.isBoundToService.getValue())
+			{
+				mServiceConnection.bindToUnrealInstanceService("CreateServiceConnection rebind");
+			}
+			return mServiceConnection;
 		}
+		
+//		if (mServiceConnection != null && mServiceConnection.isBoundToService.getValue())
+//		{
+//			return mServiceConnection;
+////			mServiceConnection.logContextDetails("CreateServiceConnection destroy prev (" + activity + ")");
+////			DeposeServiceConnection();
+//		}
 
 		if (mServiceConnection == null)
 		{
 			//mServiceConnection.logContextDetails("CreateServiceConnection create new (" + activity + ")");
-			mServiceConnection = new MyServiceConnection(activity);
+			mServiceConnection = new ASISConnection(activity);
 			mServiceConnection.bindToUnrealInstanceService("CreateServiceConnection");
 		}
 
 		return mServiceConnection;
 	}
 
-	void DeposeServiceConnection()
+	void DeposeServiceConnection(String caller)
 	{
 
-		Log.w(TAG, "DeposeServiceConnection, mServiceConnection = " + mServiceConnection);
+		Log.w(TAG, "DeposeServiceConnection, mServiceConnection = " + mServiceConnection + ", caller=" + caller);
 
-		if (mServiceConnection != null) {
-			mServiceConnection.unbindToUnrealInstanceService("DeposeServiceConnection");
-			//mServiceConnection.doCleanupForUnbinding("DeposeServiceConnection");
+		if (mServiceConnection != null && mServiceConnection.isBoundToService.getValue()) {
+			mServiceConnection.unbindToUnrealInstanceService(caller + "->DeposeServiceConnection");
+			mServiceConnection.doCleanupForUnbinding(caller + "->DeposeServiceConnection");
 		}
 
 		mServiceConnection = null;
@@ -489,17 +127,7 @@ public class UseServiceActivity extends Activity implements View.OnClickListener
 
 	final static String TAG = "UE-UseServiceActivity";
 
-	final static String obbModuleName = "Lyra";//""ALPHA_POC";
-//	final static String obbModuleName = BuildConfig.ASISModuleName; //"Lyra";//""ALPHA_POC";
-//	final static String servicePackageName = "";
-//	final static String servicePackageName = BuildConfig.ASISPackageName;
-	final static String servicePackageName = "com.epicgames.Lyra";
-	final static String obbFileLocation = ""; //"/data/local/tmp/3d/AndroidPackagingTest.main.obb.png";
-//	final static String obbFileLocation = BuildConfig.ASISOBBPath; //"/data/local/tmp/3d/AndroidPackagingTest.main.obb.png";
-
-	final static String commandLineArgs = "-nosound";
-//	final static String commandLineArgs = "-nosound -launchandroidflags=0 -dpcvars='Android.UseGameThreadForNativeCommands=0'";
-	final static boolean enablePropagateAlpha = true;
+
 	
 	private Button activateView1Button, activateView2Button, activateView3Button, restartActivityButton, themeChangeButton, resumeButton;
 	private Button pauseButton, stopButton, bindServiceButton, unbindServiceButton, memReportButton;
@@ -536,7 +164,7 @@ public class UseServiceActivity extends Activity implements View.OnClickListener
 	void stopService(int attachId, Handler.Callback callback)
 	{
 		Intent intent = new Intent();
-		intent.setClassName(servicePackageName, "com.epicgames.makeaar.UnrealSharedInstanceService");
+		intent.setClassName(ASISConnection.servicePackageName, "com.epicgames.makeaar.UnrealSharedInstanceService");
 		getApplication().stopService(intent);
 	}
 
@@ -551,13 +179,14 @@ public class UseServiceActivity extends Activity implements View.OnClickListener
 		Log.v(TAG, "UseServiceActivity: onResume");
 		super.onResume();
 
-//		GetServiceConnection("onResume").resumeService(getTaskId());
+		//GetServiceConnection("onResume").resumeService(getTaskId());
 	}
 
 	@Override
 	protected void onPause() {
 		Log.v(TAG, "UseServiceActivity: onPause");
 		super.onPause();
+		GetServiceConnection("onPause").detachSurfaceFromService(getTaskId(), null);
 	}
 
 	@Override
@@ -579,8 +208,8 @@ public class UseServiceActivity extends Activity implements View.OnClickListener
 	protected void onDestroy() {
 		Log.v(TAG, "UseServiceActivity: onDestroy");
 		super.onDestroy();
-		//DeposeServiceConnection();
-		//GetServiceConnection().doCleanupForBinding("onDestroy");
+
+		DeposeServiceConnection("onDestroy");
 	}
 
 	@Override
@@ -589,8 +218,35 @@ public class UseServiceActivity extends Activity implements View.OnClickListener
 		super.onRestart();
 	}
 
+	
 	@Override
 	protected void onCreate(Bundle savedInstanceState) {
+
+		new RelativeLayout(this);
+		
+//		getWindow().getTransitionManager().setTransition();
+//		WindowManager.LayoutParams params = getWindow().getAttributes();
+//		params.type = params.TYPE_APPLICATION_ATTACHED_DIALOG;
+//		params.flags = params.flags | Window.FEATURE_NO_TITLE;
+
+//!		getWindow();
+//!		getWindowManager();
+//!		getIntent();
+//!		runOnUiThread();
+		
+//		getResources();
+//		getExternalFilesDir();
+//		getFilesDir();
+//		getSystemService();
+//		getPackageName();
+//		getPackageResourcePath();
+//		getMainExecutor();
+//		getDisplay();
+//		getPackageManager();
+//		getApplicationContext();
+
+
+		
 		super.onCreate( savedInstanceState );
 		AppCompatDelegate.setDefaultNightMode(next_theme);
 		getTheme().applyStyle(R.style.Theme_UnrealEngine_NoActionBar, true);
@@ -621,8 +277,6 @@ public class UseServiceActivity extends Activity implements View.OnClickListener
 		stopButton.setOnClickListener( this );
 		bindServiceButton.setOnClickListener( this );
 		unbindServiceButton.setOnClickListener( this );
-		restartActivityButton.setOnClickListener( this );
-		memReportButton.setOnClickListener( this );
 		themeChangeButton.setOnClickListener( this );
 		resumeButton.setOnClickListener( this );
 
@@ -660,23 +314,24 @@ public class UseServiceActivity extends Activity implements View.OnClickListener
 			stopService(getTaskId(), null);
 		} else if (bindServiceButton.equals(view)) {
 
-//			try {
-//				GetServiceConnection("Bind Service Button").bindToUnrealInstanceService("onClick");
-//			} catch (Exception e) {
-//				e.printStackTrace();
-//			}
+			try {
+				CreateServiceConnection( this,"Bind Service Button").bindToUnrealInstanceService("onClick");
+			} catch (Exception e) {
+				e.printStackTrace();
+			}
 
 		} else if (unbindServiceButton.equals(view)) {
 
 			try {
-				DeposeServiceConnection();
+				DeposeServiceConnection("unbindServiceButton");
 				//GetServiceConnection("Unbind Service Button").unbindToUnrealInstanceService("onClick");
 			} catch (Exception e) {
 				e.printStackTrace();
 			}
 
 		} else if (restartActivityButton.equals(view)) {
-			recreate();
+
+			processRestartActivity();
 		} else if (memReportButton.equals(view)) {
 			ReportDebugMemoryInfo();
 		} else if (themeChangeButton.equals(view)) {
@@ -726,14 +381,49 @@ public class UseServiceActivity extends Activity implements View.OnClickListener
 
 	}
 
+	Handler.Callback onDetachSurfaceHandler = new Handler.Callback() {
+		@Override
+		public boolean handleMessage(Message msg) {
+			synchronized (mServiceConnection) {
+				Log.v(TAG, "onDetachSurfaceHandler::handleMessage() msg=" + msg);
+
+				recreate();
+				return true;
+			}
+		}
+	};
+	
+	private Handler handler = new Handler();
+
+	private void processRestartActivity() {
+
+		mServiceConnection.detachSurfaceFromService(getTaskId(), onDetachSurfaceHandler);
+		//recreate();
+		
+//		Runnable RestartTrigger = new Runnable() {
+//			@Override
+//			public void run() {
+//
+//				mServiceConnection.detachSurfaceFromService(getTaskId(), onDetachSurfaceHandler);
+//			}
+//		};
+//
+//		handler.post(RestartTrigger);
+	}
+
 	void doThemeChange(Integer value)
 	{
 		Log.v(TAG, "doing a touch event theme change: new uiTheme = ${value}");
 //        prev_uiMode = next_theme;
 		next_theme = value;
 //        AppCompatDelegate.setDefaultNightMode(value);
-		recreate();
-
+		if (mServiceConnection != null)
+		{
+			mServiceConnection.detachSurfaceFromService(getTaskId(), onDetachSurfaceHandler);
+		}
+		else {
+			recreate();
+		}
 	}
 
 	@Override
diff --git a/Engine/Plugins/Runtime/AndroidSingleInstanceService/Source/AndroidSingleInstanceService/Android/gradlew.bat b/Engine/Plugins/Runtime/AndroidSingleInstanceService/Source/AndroidSingleInstanceService/Android/gradlew.bat
new file mode 100755
--- /dev/null
+++ b/Engine/Plugins/Runtime/AndroidSingleInstanceService/Source/AndroidSingleInstanceService/Android/gradlew.bat
@@ -0,0 +1,89 @@
+@rem
+@rem Copyright 2015 the original author or authors.
+@rem
+@rem Licensed under the Apache License, Version 2.0 (the "License");
+@rem you may not use this file except in compliance with the License.
+@rem You may obtain a copy of the License at
+@rem
+@rem      https://www.apache.org/licenses/LICENSE-2.0
+@rem
+@rem Unless required by applicable law or agreed to in writing, software
+@rem distributed under the License is distributed on an "AS IS" BASIS,
+@rem WITHOUT WARRANTIES OR CONDITIONS OF ANY KIND, either express or implied.
+@rem See the License for the specific language governing permissions and
+@rem limitations under the License.
+@rem
+
+@if "%DEBUG%" == "" @echo off
+@rem ##########################################################################
+@rem
+@rem  Gradle startup script for Windows
+@rem
+@rem ##########################################################################
+
+@rem Set local scope for the variables with windows NT shell
+if "%OS%"=="Windows_NT" setlocal
+
+set DIRNAME=%~dp0
+if "%DIRNAME%" == "" set DIRNAME=.
+set APP_BASE_NAME=%~n0
+set APP_HOME=%DIRNAME%
+
+@rem Resolve any "." and ".." in APP_HOME to make it shorter.
+for %%i in ("%APP_HOME%") do set APP_HOME=%%~fi
+
+@rem Add default JVM options here. You can also use JAVA_OPTS and GRADLE_OPTS to pass JVM options to this script.
+set DEFAULT_JVM_OPTS="-Xmx64m" "-Xms64m"
+
+@rem Find java.exe
+if defined JAVA_HOME goto findJavaFromJavaHome
+
+set JAVA_EXE=java.exe
+%JAVA_EXE% -version >NUL 2>&1
+if "%ERRORLEVEL%" == "0" goto execute
+
+echo.
+echo ERROR: JAVA_HOME is not set and no 'java' command could be found in your PATH.
+echo.
+echo Please set the JAVA_HOME variable in your environment to match the
+echo location of your Java installation.
+
+goto fail
+
+:findJavaFromJavaHome
+set JAVA_HOME=%JAVA_HOME:"=%
+set JAVA_EXE=%JAVA_HOME%/bin/java.exe
+
+if exist "%JAVA_EXE%" goto execute
+
+echo.
+echo ERROR: JAVA_HOME is set to an invalid directory: %JAVA_HOME%
+echo.
+echo Please set the JAVA_HOME variable in your environment to match the
+echo location of your Java installation.
+
+goto fail
+
+:execute
+@rem Setup the command line
+
+set CLASSPATH=%APP_HOME%\gradle\wrapper\gradle-wrapper.jar
+
+
+@rem Execute Gradle
+"%JAVA_EXE%" %DEFAULT_JVM_OPTS% %JAVA_OPTS% %GRADLE_OPTS% "-Dorg.gradle.appname=%APP_BASE_NAME%" -classpath "%CLASSPATH%" org.gradle.wrapper.GradleWrapperMain %*
+
+:end
+@rem End local scope for the variables with windows NT shell
+if "%ERRORLEVEL%"=="0" goto mainEnd
+
+:fail
+rem Set variable GRADLE_EXIT_CONSOLE if you need the _script_ return code instead of
+rem the _cmd.exe /c_ return code!
+if  not "" == "%GRADLE_EXIT_CONSOLE%" exit 1
+exit /b 1
+
+:mainEnd
+if "%OS%"=="Windows_NT" endlocal
+
+:omega
diff --git a/Engine/Plugins/Runtime/AndroidSingleInstanceService/Source/AndroidSingleInstanceService/Android/setup.bat b/Engine/Plugins/Runtime/AndroidSingleInstanceService/Source/AndroidSingleInstanceService/Android/setup.bat
new file mode 100755
--- /dev/null
+++ b/Engine/Plugins/Runtime/AndroidSingleInstanceService/Source/AndroidSingleInstanceService/Android/setup.bat
@@ -0,0 +1,7 @@
+@echo off
+@echo Creating local.properties
+setlocal ENABLEDELAYEDEXPANSION
+set NDKDIR=!NDKROOT:\=\\!
+set SDKDIR=!ANDROID_HOME:\=\\!
+echo ndk.dir=%NDKDIR%> local.properties
+echo sdk.dir=%SDKDIR%>> local.properties
\ No newline at end of file
diff --git a/Engine/Plugins/Runtime/AndroidSingleInstanceService/Source/AndroidSingleInstanceService/Android/setup.sh b/Engine/Plugins/Runtime/AndroidSingleInstanceService/Source/AndroidSingleInstanceService/Android/setup.sh
new file mode 100644
--- /dev/null
+++ b/Engine/Plugins/Runtime/AndroidSingleInstanceService/Source/AndroidSingleInstanceService/Android/setup.sh
@@ -0,0 +1,12 @@
+NDKDIR=$NDKROOT;
+SDKDIR=$ANDROID_HOME;
+destdir=local.properties
+
+if [ -f "$destdir" ]
+then 
+    echo -e "ndk.dir=$NDKDIR\nsdk.dir=$SDKDIR" > "$destdir"
+    echo "local.properties has been updated"
+else
+    echo -e "ndk.dir=$NDKDIR\nsdk.dir=$SDKDIR" > "$destdir"
+    echo "local.properties has been created"
+fi
diff --git a/Engine/Plugins/Runtime/AndroidSingleInstanceService/Source/AndroidSingleInstanceService/Android/src/main/java/com/epicgames/makeaar/Engine.java b/Engine/Plugins/Runtime/AndroidSingleInstanceService/Source/AndroidSingleInstanceService/Android/src/main/java/com/epicgames/makeaar/Engine.java
--- a/Engine/Plugins/Runtime/AndroidSingleInstanceService/Source/AndroidSingleInstanceService/Android/src/main/java/com/epicgames/makeaar/Engine.java
+++ b/Engine/Plugins/Runtime/AndroidSingleInstanceService/Source/AndroidSingleInstanceService/Android/src/main/java/com/epicgames/makeaar/Engine.java
@@ -494,6 +494,11 @@ public class Engine
 	//}
 }
 
+void IssueAppCommand(int cmd, String caller)
+{
+	GetOverrideGameActivityWhenInitComplete().nativeAppCommand(cmd);
+}
+
 void DoGameActivityNativeMain(String projectModule, String reasonString)
 	{
 		Log.i(TAG, "DoGameActivityNativeMain is called with reasonString=" + reasonString + ", projectModule=" + projectModule);
diff --git a/Engine/Plugins/Runtime/AndroidSingleInstanceService/Source/AndroidSingleInstanceService/Android/src/main/java/com/epicgames/makeaar/GameActivityForMakeAAR.java b/Engine/Plugins/Runtime/AndroidSingleInstanceService/Source/AndroidSingleInstanceService/Android/src/main/java/com/epicgames/makeaar/GameActivityForMakeAAR.java
--- a/Engine/Plugins/Runtime/AndroidSingleInstanceService/Source/AndroidSingleInstanceService/Android/src/main/java/com/epicgames/makeaar/GameActivityForMakeAAR.java
+++ b/Engine/Plugins/Runtime/AndroidSingleInstanceService/Source/AndroidSingleInstanceService/Android/src/main/java/com/epicgames/makeaar/GameActivityForMakeAAR.java
@@ -7,9 +7,12 @@ import android.content.ContextWrapper;
 import java.io.File;
 import java.util.Map;
 import java.util.HashMap;
+import java.util.concurrent.Executor;
 
+import android.content.Intent;
 import android.content.pm.ApplicationInfo;
 import android.content.pm.FeatureInfo;
+import android.content.res.Resources;
 import android.os.Bundle;
 import android.os.Looper;
 
@@ -22,8 +25,14 @@ import android.content.pm.PackageInfo;
 
 import android.media.AudioManager;
 
+import android.view.Display;
 import android.view.View;
 import android.view.Surface;
+import android.view.Window;
+import android.view.WindowManager;
+
+import androidx.annotation.NonNull;
+import androidx.annotation.Nullable;
 
 import com.epicgames.unreal.GameActivity;
 import com.epicgames.unreal.Logger;
@@ -76,8 +85,69 @@ public class GameActivityForMakeAAR extends com.epicgames.unreal.GameActivity
 	{
 		return Instance();
 	}
-
 	
+	@Override
+	public  @Nullable Object getSystemService(@NonNull String name) 
+	{
+		return GetCurrentActivityContext().getSystemService(name);
+	}
+
+	@Override
+	public Resources getResources()
+	{
+		return GetCurrentActivityContext().getResources();
+	}
+
+	@Override
+    public @Nullable File getExternalFilesDir(@Nullable String type) 
+	{
+		return GetCurrentActivityContext().getExternalFilesDir(type);
+    }
+
+	@Override
+	public File getFilesDir()
+	{
+		return GetCurrentActivityContext().getFilesDir();
+	}
+
+	@Override
+	public String getPackageName()
+	{
+		return GetCurrentActivityContext().getPackageName();
+	}
+
+	@Override
+	public String getPackageResourcePath()
+	{
+		return GetCurrentActivityContext().getPackageResourcePath();
+	}
+
+	@Override
+	public Executor getMainExecutor()
+	{
+		return GetCurrentActivityContext().getMainExecutor();
+	}
+
+	@Override
+	public @Nullable Display getDisplay()
+	{
+		return GetCurrentActivityContext().getDisplay();
+	}
+
+	@Override
+	public PackageManager getPackageManager()
+	{
+		return GetCurrentActivityContext().getPackageManager();
+	}
+
+	@Override
+	public Context getApplicationContext()
+	{
+		return GetCurrentActivityContext().getApplicationContext();
+	}
+
+
+
 	void setCommandline(String inCommandline)
 	{
 		nativeSetCommandline(inCommandline);
@@ -88,14 +158,6 @@ public class GameActivityForMakeAAR extends com.epicgames.unreal.GameActivity
 		return (GameActivityForMakeAAR)_gameActivityInstance;
 	}
 
-	// already present in super call and activity!
-	//public final void runOnUiThread(Runnable action)
-	//{
-	//	assert(activityContext != null);
-	//	activityContext.runOnUiThread(action);
-	//}
-
-
 	@Override
 	protected void RestartApplication(String RestartExtra)
 	{
@@ -340,6 +402,10 @@ public class GameActivityForMakeAAR extends com.epicgames.unreal.GameActivity
 
 	public boolean setActivity(SimpleContextWrapper inActivity, GameActivitySetupInfo gameActivitySetupInfo)
 	{
+
+		// pre super.onCreate {
+		//
+
 		if (activityContext != null)
 		{
 			// we need to shut down the old activity
@@ -357,15 +423,26 @@ public class GameActivityForMakeAAR extends com.epicgames.unreal.GameActivity
 		}
 		
 
+		assert(_gameActivityInstance != null);
+		bOnCreateCalled = true;
 
-		
-		clipboardManager = gameActivitySetupInfo.clipboardManager;
+		InternalFilesDir = gameActivitySetupInfo.InternalFilesDir;
+		ExternalFilesDir = gameActivitySetupInfo.ExternalFilesDir;
+//$${gameActivityOnCreateBeginningAdditions}$$
+		Logger.RegisterCallback(this);
 
 		// Grab a reference to the asset manager
 		AssetManagerReference = gameActivitySetupInfo.AssetManagerReference;
 
-		InternalFilesDir = gameActivitySetupInfo.InternalFilesDir;
-		ExternalFilesDir = gameActivitySetupInfo.ExternalFilesDir;
+	//
+	//} pre super.onCreate
+	//
+
+		Bundle savedInstanceState = new Bundle();
+		onCreateBody(savedInstanceState);
+
+		
+		clipboardManager = gameActivitySetupInfo.clipboardManager;
 
 		_extrasBundle = gameActivitySetupInfo._extrasBundle;
 
diff --git a/Engine/Plugins/Runtime/AndroidSingleInstanceService/Source/AndroidSingleInstanceService/Android/src/main/java/com/epicgames/makeaar/UnrealSharedInstanceService.java b/Engine/Plugins/Runtime/AndroidSingleInstanceService/Source/AndroidSingleInstanceService/Android/src/main/java/com/epicgames/makeaar/UnrealSharedInstanceService.java
--- a/Engine/Plugins/Runtime/AndroidSingleInstanceService/Source/AndroidSingleInstanceService/Android/src/main/java/com/epicgames/makeaar/UnrealSharedInstanceService.java
+++ b/Engine/Plugins/Runtime/AndroidSingleInstanceService/Source/AndroidSingleInstanceService/Android/src/main/java/com/epicgames/makeaar/UnrealSharedInstanceService.java
@@ -37,10 +37,40 @@ public class UnrealSharedInstanceService extends Service {
 	private static final String TAG = "UE_ASIS_SharedInstanceService";
 
 	static int activeTaskID = 0;
+
+	private static Engine engineInstance;
+
+	public static String defaultProjectModuleName = "UNDEFINED_PackageName";
+	//		public static String defaultProjectModuleName = "Alpha_POC";
+	public static String defaultOBBFileLocation = "";//"Gallery3D"; //"Alpha_POC";
+													 //	public final static String defaultProjectModuleName = "";
+													 //		public static String defaultCommandLine = "../../../" + BuildConfig.OBB_MODULE_NAME + "/" + BuildConfig.OBB_MODULE_NAME + ".uproject -nosound -nocrashreports";
+	public static String defaultCommandLine = "-nosound -nocrashreports -dpcvars=a.UseSwappyForFramePacing=0 -tracehost=127.0.0.1 -cpuprofilertrace -statnamedevents -trace=Bookmark,Frame,CPU,GPU,LoadTime,File";
+//		public final static Engine gUnrealEngine = EngineFactory.getInstance(null, defaultOBBFileLocation, defaultProjectModuleName, true);
+//		static Engine gUnrealEngine;
+
+	private static final int APP_CMD_INPUT_CHANGED = 0;
+	private static final int APP_CMD_INIT_WINDOW = 1;
+	private static final int APP_CMD_TERM_WINDOW = 2;
+	private static final int APP_CMD_WINDOW_RESIZED = 3;
+	private static final int APP_CMD_WINDOW_REDRAW_NEEDED = 4;
+	private static final int APP_CMD_CONTENT_RECT_CHANGED = 5;
+	private static final int APP_CMD_GAINED_FOCUS = 6;
+	private static final int APP_CMD_LOST_FOCUS = 7;
+	private static final int APP_CMD_CONFIG_CHANGED = 8;
+	private static final int APP_CMD_LOW_MEMORY = 9;
+	private static final int APP_CMD_START = 10;
+	private static final int APP_CMD_RESUME = 11;
+	private static final int APP_CMD_SAVE_STATE = 12;
+	private static final int APP_CMD_PAUSE = 13;
+	private static final int APP_CMD_STOP = 14;
+	private static final int APP_CMD_DESTROY = 15;
+
 	/**
 	* Target we publish for clients to send messages to IncomingHandler.
 	*/
 	Messenger mMessenger;
+
 	public String Name()
 	{
 		return getApplication().getClass().getSimpleName();
@@ -50,30 +80,28 @@ public class UnrealSharedInstanceService extends Service {
 		return context.getPackageName() +",context=" + context.toString();
 	}
 
+	final static void SetOBB(String commandLineArgs, String obbModuleName, String obbFileLocation)
+	{
+		if (commandLineArgs.isEmpty())
+			commandLineArgs = "-nosound -nocrashreports";
+		if (!obbModuleName.isEmpty()) {
+			defaultProjectModuleName = obbModuleName;
+			defaultCommandLine = "../../../" + defaultProjectModuleName + "/" + defaultProjectModuleName + ".uproject " + commandLineArgs;
+//			if (!obbFileLocation.isEmpty())
+			defaultOBBFileLocation = obbFileLocation;
+		}
+	}
+	
 	public UnrealSharedInstanceService() {
 	}
 
 	/**
 	 * Handler of incoming messages from clients.
 	 */
-	static class IncomingHandler extends Handler implements Engine.IEventCallback
+	class IncomingHandler extends Handler implements Engine.IEventCallback
 	{
 		private Service applicationContext;
-		static private Engine engineInstance;
-
-
-		public static String defaultProjectModuleName = "UNDEFINED_PackageName";
-//		public static String defaultProjectModuleName = "Alpha_POC";
-		public static String defaultOBBFileLocation = "";//"Gallery3D"; //"Alpha_POC";
-		//	public final static String defaultProjectModuleName = "";
-//		public static String defaultCommandLine = "../../../" + BuildConfig.OBB_MODULE_NAME + "/" + BuildConfig.OBB_MODULE_NAME + ".uproject -nosound -nocrashreports";
-		public static String defaultCommandLine = "-nosound -nocrashreports -dpcvars=a.UseSwappyForFramePacing=0 -tracehost=127.0.0.1 -cpuprofilertrace -statnamedevents -trace=Bookmark,Frame,CPU,GPU,LoadTime,File";
-//		public final static Engine gUnrealEngine = EngineFactory.getInstance(null, defaultOBBFileLocation, defaultProjectModuleName, true);
-//		static Engine gUnrealEngine;
-		public static Messenger mServiceReply;
-
-
-
+		public Messenger mServiceReply;
 
 		void Release()
 		{
@@ -81,40 +109,18 @@ public class UnrealSharedInstanceService extends Service {
 			mServiceReply = null;
 		}
 
-		final static void SetOBB(String commandLineArgs, String obbModuleName, String obbFileLocation)
-		{
-			if (commandLineArgs.isEmpty())
-				commandLineArgs = "-nosound -nocrashreports";
-			if (!obbModuleName.isEmpty()) {
-				defaultProjectModuleName = obbModuleName;
-				defaultCommandLine = "../../../" + defaultProjectModuleName + "/" + defaultProjectModuleName + ".uproject " + commandLineArgs;
-//			if (!obbFileLocation.isEmpty())
-				defaultOBBFileLocation = obbFileLocation;
-			}
 
-		}
 
-		private static int APP_CMD_INPUT_CHANGED = 0;
-		private static int APP_CMD_INIT_WINDOW = 1;
-		private static int APP_CMD_TERM_WINDOW = 2;
-		private static int APP_CMD_WINDOW_RESIZED = 3;
-		private static int APP_CMD_WINDOW_REDRAW_NEEDED = 4;
-		private static int APP_CMD_CONTENT_RECT_CHANGED = 5;
-		private static int APP_CMD_GAINED_FOCUS = 6;
-		private static int APP_CMD_LOST_FOCUS = 7;
-		private static int APP_CMD_CONFIG_CHANGED = 8;
-		private static int APP_CMD_LOW_MEMORY = 9;
-		private static int APP_CMD_START = 10;
-		private static int APP_CMD_RESUME = 11;
-		private static int APP_CMD_SAVE_STATE = 12;
-		private static int APP_CMD_PAUSE = 13;
-		private static int APP_CMD_STOP = 14;
-		private static int APP_CMD_DESTROY = 15;
+
 		
 		IncomingHandler(Service context, Intent intent) {
 		    mServiceReply = intent.getParcelableExtra("callback");
 			applicationContext = context;
 
+			Log.w(TAG, "LIFECYCLE: IncomingHandler(constructor intent) - proc=" + Application.getProcessName()
+				+ ", mServiceReply=" + mServiceReply
+				+ ", applicationContext=" + applicationContext);
+
 		    if (intent.hasExtra("obbModuleName"))
 		    {
 			    Bundle data = intent.getExtras();
@@ -123,7 +129,7 @@ public class UnrealSharedInstanceService extends Service {
 			    String commandLineArgs = data.getString("commandLineArgs", "");
 			    SetOBB(commandLineArgs, obbModuleName, obbFileLocation);
     
-			    Log.i("UESharedInstanceServiceNative", "**\tIncomingHandler(constructor intent) - proc=" + Name(context)
+			    Log.i(TAG, "**\tIncomingHandler(constructor intent) - proc=" + Name(context)
 				    + ", obbModuleName=" + obbModuleName
 				    + ", obbFileLocation=" + obbFileLocation);
 		    }
@@ -205,9 +211,8 @@ public class UnrealSharedInstanceService extends Service {
 		}
 
 		@Override
-		public void handleMessage(Message msg) {
-
-
+		public void handleMessage(Message msg) 
+		{
 			if (msg.what == UnrealMessageType.Hello.ordinal()) {
 				Toast.makeText(applicationContext, "hello!", Toast.LENGTH_SHORT).show();
 			}
@@ -216,8 +221,8 @@ public class UnrealSharedInstanceService extends Service {
 				Bundle data = msg.getData();
 				Surface externalSurface = data.getParcelable("surface");
 				int taskId = data.getInt("taskId", 0);
-				boolean bDoStart = data.getBoolean("start", true);
-				boolean bDoResume = data.getBoolean("resume", true);
+				boolean bDoStart = data.getBoolean("start", false);
+				boolean bDoResume = data.getBoolean("resume", false);
 
 				Log.i(TAG, "handleMessage(MSG_ATTACH_EXTERNAL_SURFACE) - proc=" + Name(this.applicationContext)
 					+ ", taskId=" + taskId
@@ -241,9 +246,9 @@ public class UnrealSharedInstanceService extends Service {
 
 				engineInstance.registerEventCallback(this);
 
-				//if (bDoStart)
+				if (bDoStart)
 				{
-					//engineInstance.QueueStart("UnrealMessageType.AttachExternalSurface, taskId=" + taskId); //777
+					engineInstance.QueueStart("UnrealMessageType.AttachExternalSurface, taskId=" + taskId);
 				}
 
 				if (externalSurface != null && externalSurface.isValid() && engineInstance != null) {
@@ -268,7 +273,7 @@ public class UnrealSharedInstanceService extends Service {
 					}
 
 					
-					//if (bDoResume)
+					if (bDoResume)
 					{
 						engineInstance.QueueResume("UnrealMessageType.AttachExternalSurface, taskId=" + taskId);
 					}
@@ -304,8 +309,8 @@ public class UnrealSharedInstanceService extends Service {
 				Bundle data = msg.getData();
 
 				int taskId = data.getInt("taskId", 0);
-				boolean bDoStop = data.getBoolean("stop", true);
-				boolean bDoPause = data.getBoolean("pause", true);
+				boolean bDoStop = data.getBoolean("stop", false);
+				boolean bDoPause = data.getBoolean("pause", false);
 
 				Log.i(TAG, "handleMessage(MSG_DETACH_EXTERNAL_SURFACE) - proc=" + Name(applicationContext)
 					+ ", taskId=" + taskId
@@ -315,21 +320,45 @@ public class UnrealSharedInstanceService extends Service {
 				//if (taskId != 0 && taskId == activeTaskID )
 				{
 
-					//if (bDoPause)
+					if (bDoPause)
 					{
 						engineInstance.QueuePause("UnrealMessageType.DetachExternalSurface, taskId=" + taskId);
 					}
 
-					//if (bDoStop)
+					if (bDoStop)
 					{
 						engineInstance.QueueStop("UnrealMessageType.DetachExternalSurface, taskId=" + taskId);
 					}
 
 					//engineInstance.sendConsoleCommand("t.maxfps 0.001");
+					activeTaskID = 0;
+				}
+				
+				if (msg.replyTo != null) {
+					Log.w(TAG, "Sending message for detach surface back to message replyTo obj=" + engineInstance);
+					Message replyMsg = Message.obtain();
+					replyMsg.what = Engine.EVENTTYPE_ENGINELOOP_SUSPENDED;
+					replyMsg.arg1 = activeTaskID;
+					try {
+						msg.replyTo.send(replyMsg);
+					} catch (RemoteException e) {
+						throw new RuntimeException(e);
+					}
+				}
 					
+				if (mServiceReply != null && mServiceReply != msg.replyTo)
+				{
+					Log.w(TAG, "Sending message for detach surface back to mServiceReply=" + mServiceReply);
+					Message replyMsg = Message.obtain();
+					replyMsg.what = Engine.EVENTTYPE_ENGINELOOP_SUSPENDED;
+					replyMsg.arg1 = activeTaskID;
+					try {
+						mServiceReply.send(replyMsg);
+					} catch (RemoteException e) {
+						throw new RuntimeException(e);
+					}
 				}
-				activeTaskID = taskId;
-
+				
 				//mServiceReply = null;
 				//engineInstance.registerEventCallback(null);
 				//applicationContext.stopSelf();
@@ -347,8 +376,7 @@ public class UnrealSharedInstanceService extends Service {
 
 				engineInstance.QueuePause("UnrealMessageType.StopService, taskId=" + taskId);
 				engineInstance.QueueStop("UnrealMessageType.StopService, taskId=" + taskId);
-				//engineInstance.onPause(engineInstance.getCurrentContextID(), "UnrealMessageType.StopService, taskId=" + taskId);
-				//engineInstance.onStop(engineInstance.getCurrentContextID(), "UnrealMessageType.StopService, taskId=" + taskId);
+
 
 				activeTaskID = taskId;
 
@@ -360,7 +388,10 @@ public class UnrealSharedInstanceService extends Service {
 			{
 				Bundle data = msg.getData();
 				int taskId = data.getInt("taskId", 0);
-				boolean bResume = data.getBoolean("resume", true);
+				boolean bResume = data.getBoolean("resume", false);
+				boolean bCMDResume = data.getBoolean("cmd_resume", false);
+				boolean bCMDFocus = data.getBoolean("cmd_focus", false);
+				boolean bCMDInitWindow = data.getBoolean("cmd_initWindow", false);
 				boolean bCMDMainInit = data.getBoolean("cmd_mainInit", false);
 
 				Log.i(TAG, "handleMessage(MSG_RESUME_SERVICE) - proc=" + Name(applicationContext)
@@ -376,10 +407,26 @@ public class UnrealSharedInstanceService extends Service {
 				}
 				else {
 
+					if (bCMDResume)
+					{
+						GameActivityForMakeAAR.nativeAppCommand(APP_CMD_RESUME);
+					}
+
+					if (bCMDFocus)
+					{
+						GameActivityForMakeAAR.nativeAppCommand(APP_CMD_GAINED_FOCUS);
+					}
+
+					if (bCMDInitWindow)
+					{
+						GameActivityForMakeAAR.nativeAppCommand(APP_CMD_INIT_WINDOW);
+					}
+
 					if (bCMDMainInit)
+					{
 						GameActivity.Get().nativeResumeMainInit();
+					}	
 				}
-
 			}
 			else if (msg.what == UnrealMessageType.TouchEvent.ordinal())
 			{
@@ -406,6 +453,7 @@ public class UnrealSharedInstanceService extends Service {
 		@Override
 		protected void finalize() throws Throwable {
 			super.finalize();
+			Log.e(TAG, "LIFECYCLE: IncomingHandler(finalize) mServiceReply=" + mServiceReply);
 			mServiceReply = null;
 			engineInstance.registerEventCallback(null);
 		}
diff --git a/Engine/Source/Runtime/ApplicationCore/Private/Android/AndroidWindow.cpp b/Engine/Source/Runtime/ApplicationCore/Private/Android/AndroidWindow.cpp
--- a/Engine/Source/Runtime/ApplicationCore/Private/Android/AndroidWindow.cpp
+++ b/Engine/Source/Runtime/ApplicationCore/Private/Android/AndroidWindow.cpp
@@ -160,7 +160,7 @@ JNI_METHOD void Java_com_epicgames_unreal_GameActivity_nativeSetWindowInfo(JNIEn
 
 JNI_METHOD void Java_com_epicgames_unreal_GameActivity_nativeSetSurfaceViewInfo(JNIEnv* jenv, jobject thiz, jint width, jint height)
 {
-	STANDALONE_DEBUG_LOG( TEXT("nativeSetSurfaceViewInfo prev width=%d and prev height=%d"), GSurfaceViewWidth, GSurfaceViewHeight);
+	STANDALONE_DEBUG_LOG( TEXT("nativeSetSurfaceViewInfo prev[width=%d, height=%d] new[width=%d, height=%d]"), GSurfaceViewWidth, GSurfaceViewHeight, width, height);
 
 	if (GAndroidWindowOverride != nullptr && (width != GSurfaceViewWidth || height != GSurfaceViewHeight))	
 	{
@@ -332,13 +332,15 @@ void FAndroidWindow::ReleaseWindowRef(ANativeWindow* InWindow)
  void FAndroidWindow::SetHardwareWindow_EventThread(void* InWindow)
 {
 #if USE_ANDROID_EVENTS
+	 STANDALONE_DEBUG_LOGf(LogAndroid, TEXT("SetHardwareWindow_EventThread(USE_ANDROID_EVENTS) -> InWindow(%p), GAndroidWindowOverride(%p), IsInAndroidEventThread()=%d"), InWindow, GAndroidWindowOverride, IsInAndroidEventThread());
+
 	check(IsInAndroidEventThread());
 #endif
 
 #if USE_ANDROID_STANDALONE
 	if (GAndroidWindowOverride && InWindow != GAndroidWindowOverride)
 	{
-		STANDALONE_DEBUG_LOGf(LogAndroid, TEXT("SetHardwareWindow_EventThread -> InWindow(%p) is not current GAndroidWindowOverride(%p)"), InWindow, GAndroidWindowOverride);
+		STANDALONE_DEBUG_LOGf(LogAndroid, TEXT("SetHardwareWindow_EventThread(USE_ANDROID_STANDALONE) -> InWindow(%p) is not current GAndroidWindowOverride(%p)"), InWindow, GAndroidWindowOverride);
 	}
 #endif
 
@@ -349,15 +351,22 @@ void FAndroidWindow::ReleaseWindowRef(ANativeWindow* InWindow)
 void* FAndroidWindow::GetHardwareWindow_EventThread()
 {
 #if USE_ANDROID_STANDALONE
-	if (NativeWindow && GAndroidWindowOverride)
+	void* result = GAndroidWindowOverride != nullptr ? GAndroidWindowOverride : NativeWindow;
+	if (result != nullptr)
 	{
-		STANDALONE_DEBUG_LOG(TEXT("GetHardwareWindow_EventThread overrode window: %p"), GAndroidWindowOverride);
-		return GAndroidWindowOverride;
+		if (GAndroidWindowOverride != NativeWindow)
+		{
+			STANDALONE_DEBUG_LOGf(LogAndroid, TEXT("GetHardwareWindow_EventThread GAndroidWindowOverride=%p, NativeWindow=%p"), GAndroidWindowOverride, NativeWindow);
+		}
+		return result;
+	}
+	else
+	{
+		//STANDALONE_DEBUG_LOGf(LogAndroid, TEXT("ERROR: GetHardwareWindow_EventThread has invalid window!!! GAndroidWindowOverride=%p, NativeWindow=%p"), GAndroidWindowOverride, NativeWindow);
+		return result;
 	}
 #endif
 
-	STANDALONE_DEBUG_LOG(TEXT("GetHardwareWindow_EventThread NativeWindow: %p"), NativeWindow);
-
 	return NativeWindow;
 }
 
@@ -408,6 +417,9 @@ void FAndroidWindow::EventManagerUpdateWindowDimensions(int32 Width, int32 Heigh
 	check(Width >= 0 && Height >= 0);
 
 #if USE_ANDROID_STANDALONE
+	STANDALONE_DEBUG_LOGf(LogAndroid, TEXT("FAndroidWindow::EventManagerUpdateWindowDimensions GAndroidWindowOverride=%p, Width=%d, Height=%d, GSurfaceViewWidth=%d, GSurfaceViewHeight=%d, CachedNativeWindowWidth=%d, CachedNativeWindowHeight=%d"), 
+		GAndroidWindowOverride, Width, Height, GSurfaceViewWidth, GSurfaceViewHeight, CachedNativeWindowWidth, CachedNativeWindowHeight);
+
 	if (GAndroidWindowOverride && GSurfaceViewWidth > 0)
 	{
 		Width = GSurfaceViewWidth;
@@ -629,6 +641,18 @@ FPlatformRect FAndroidWindow::GetScreenRect(bool bUseEventThreadWindow)
 			}
 		}
 
+
+		bool bIsPortrait = GDeviceScreenOrientation == EDeviceScreenOrientation::Portrait || GDeviceScreenOrientation == EDeviceScreenOrientation::PortraitUpsideDown;
+		UE_LOG(LogAndroid, Log, TEXT("FAndroidWindow::GetScreenRect bIsPortrait=%d"), bIsPortrait);
+
+		if (bIsPortrait)
+		{
+			UE_LOG(LogAndroid, Log, TEXT("FAndroidWindow::GetScreenRect(swap WH) bIsPortrait=%d"), bIsPortrait);
+			int temp = ScreenWidth;
+			ScreenWidth = ScreenHeight;
+			ScreenHeight = temp;
+		}
+
 		// save for future calls
 		CurrentParams.WindowWidth = ScreenWidth;
 		CurrentParams.WindowHeight = ScreenHeight;
@@ -662,6 +686,7 @@ void FAndroidWindow::CalculateSurfaceSize(int32_t& SurfaceWidth, int32_t& Surfac
 	UE_LOG(LogAndroid, Fatal, TEXT("FAndroidWindow::CalculateSurfaceSize currently expects non-JNI platforms to override resolution"));
 
 #else
+	STANDALONE_DEBUG_LOGf(LogAndroid, TEXT("::CalculateSurfaceSize(USE_ANDROID_JNI) -> bUseEventThreadWindow=%d, GAndroidWindowOverride(%p), IsInAndroidEventThread()=%d"), bUseEventThreadWindow, GAndroidWindowOverride, IsInAndroidEventThread());
 
 	if (bUseEventThreadWindow)
 	{
diff --git a/Engine/Source/Runtime/Launch/Private/Android/AndroidEventManager.cpp b/Engine/Source/Runtime/Launch/Private/Android/AndroidEventManager.cpp
--- a/Engine/Source/Runtime/Launch/Private/Android/AndroidEventManager.cpp
+++ b/Engine/Source/Runtime/Launch/Private/Android/AndroidEventManager.cpp
@@ -299,7 +299,7 @@ void FAppEventManager::ExecWindowCreated()
 	// When application launched while device is in sleep mode SystemResolution could be set to opposite orientation values
 	// Force to update SystemResolution to current values whenever we create a new window
 	FPlatformRect ScreenRect = FAndroidWindow::GetScreenRect();
-	FSystemResolution::RequestResolutionChange(ScreenRect.Right, ScreenRect.Bottom, EWindowMode::Fullscreen);
+	FSystemResolution::RequestResolutionChange(ScreenRect.Right - ScreenRect.Left, ScreenRect.Bottom - ScreenRect.Top, EWindowMode::Fullscreen);
 
 	// ReInit with the new window handle
 	FAndroidAppEntry::ReInitWindow();
diff --git a/Engine/Source/Runtime/Launch/Private/Android/LaunchAndroid.cpp b/Engine/Source/Runtime/Launch/Private/Android/LaunchAndroid.cpp
--- a/Engine/Source/Runtime/Launch/Private/Android/LaunchAndroid.cpp
+++ b/Engine/Source/Runtime/Launch/Private/Android/LaunchAndroid.cpp
@@ -281,7 +281,6 @@ static void IssueConsoleCommand(FString Command)
 // The event thread locks this whenever the window unavailable during early init, pause and resume.
 FCriticalSection GAndroidWindowLock;
 static bool bAppIsActive_EventThread = false;
-static bool bDidGainFocus = false;
 static bool bHasFocus = false;
 static bool bHasWindow = false;
 static bool bIsResumed = false;
@@ -361,18 +360,19 @@ void GAndroidWindowLock_Unlock(FString calledBy)
 static void SuspendApp_EventThread();
 static void ActivateApp_EventThread();
 #if USE_ANDROID_STANDALONE && !(UE_BUILD_SHIPPING || UE_BUILD_TEST)
-#define DEVELOPER_LOG_COMMANDCB_CASE(a) STANDALONE_DEBUG_LOGf(LogAndroid, TEXT("--TRACE-- OnAppCommandCB - Case [%s] : bDidGainFocus = %d, bHasFocus = %d, bHasWindow = %d, bIsResumed = %d, bShouldRestartFromInterrupt = %d, bReadyToProcessEvents=%d, bAppIsActive_EventThread=%d"), TEXT(#a), bDidGainFocus, bHasFocus, bHasWindow, bIsResumed, bShouldRestartFromInterrupt, bReadyToProcessEvents, bAppIsActive_EventThread)
+#define DEVELOPER_LOG_COMMANDCB_CASE(a) STANDALONE_DEBUG_LOGf(LogAndroid, TEXT("--TRACE-- OnAppCommandCB - Case [%s] : bHasFocus = %d, bHasWindow = %d, bIsResumed = %d, bShouldRestartFromInterrupt = %d, bReadyToProcessEvents=%d, bAppIsActive_EventThread=%d"), TEXT(#a), bHasFocus, bHasWindow, bIsResumed, bShouldRestartFromInterrupt, bReadyToProcessEvents, bAppIsActive_EventThread)
 #else
 #define DEVELOPER_LOG_COMMANDCB_CASE(a) 
 #endif
 
 void FPlatformMisc::UnlockAndroidWindow()
 {
+	DEVELOPER_LOG_COMMANDCB_CASE(UnlockAndroidWindow);
+	UE_LOG(LogAndroid, Log, TEXT("Unlocking android HW window during preinit. bAppIsActive_EventThread=%d"), bAppIsActive_EventThread);
+
 	check(IsInGameThread());
 	check(FTaskGraphInterface::IsRunning());
 
-	DEVELOPER_LOG_COMMANDCB_CASE(UnlockAndroidWindow);
-	UE_LOG(LogAndroid, Log, TEXT("Unlocking android HW window during preinit. bAppIsActive_EventThread=%d"), bAppIsActive_EventThread);
 	GAndroidWindowLock_Unlock("UnlockAndroidWindow");
 	
 	bReadyToProcessEvents = true;
@@ -380,6 +380,8 @@ void FPlatformMisc::UnlockAndroidWindow()
 
 JNI_METHOD void Java_com_epicgames_unreal_GameActivity_nativeResumeMainInit(JNIEnv* jenv, jobject thiz)
 {
+	STANDALONE_DEBUG_LOGf(LogAndroid, TEXT("nativeResumeMainInit() called, prev GResumeMainInit=%d, GEventHandlerInitialized=%d"), GResumeMainInit, GEventHandlerInitialized);
+
 	GResumeMainInit = true;
 
 	// now wait for event handler to be set up before returning
@@ -592,6 +594,10 @@ static void IssueAppCommand(long Command);
 void* AndroidMain(void* param)
 {
 	struct android_app* state = (struct android_app*)param;
+	FTaskTagScope Scope(ETaskTag::EGameThread);
+
+	GGameThreadId = FPlatformTLS::GetCurrentThreadId();
+
 #else
 int32 AndroidMain(struct android_app* state)
 {
@@ -776,9 +782,7 @@ int32 AndroidMain(struct android_app* state)
 
 		UE_LOG(LogAndroid, Log, TEXT("PreInit android HW window lock. bAppIsActive_EventThread=%d"), bAppIsActive_EventThread);
 		DEVELOPER_LOG_COMMANDCB_CASE(PreInitWindowLock);
-#if !USE_ANDROID_STANDALONE
 		GAndroidWindowLock_Lock("AndroidMain lock");
-#endif
 	}
 
 	FPlatformMisc::LowLevelOutputDebugString( TEXT("After GAndroidWindowLock in AndroidMain"));
@@ -860,7 +864,11 @@ int32 AndroidMain(struct android_app* state)
 		FAndroidStats::UpdateAndroidStats();
 
 		FAppEventManager::GetInstance()->Tick();
-		if (!FAppEventManager::GetInstance()->IsGamePaused())
+		if (!FAppEventManager::GetInstance()->IsGamePaused()
+#if USE_ANDROID_STANDALONE
+			&& FAppEventManager::GetInstance()->IsGameInFocus()
+#endif
+			)
 		{
 			GEngineLoop.Tick();
 		}
@@ -1036,6 +1044,7 @@ static void* AndroidEventThreadWorker( void* param )
 	FPlatformProcess::SetThreadAffinityMask(FPlatformAffinity::GetMainGameMask());
 
 	FPlatformMisc::LowLevelOutputDebugString(TEXT("Entering event processing thread engine entry point"));
+	STANDALONE_DEBUG_LOGf(LogAndroid, TEXT("AndroidEventThreadWorker(started), EventThreadID=%d"), EventThreadID);
 
 	ALooper* looper = ALooper_prepare(ALOOPER_PREPARE_ALLOW_NON_CALLBACKS);
 	ALooper_addFd(looper, state->msgread, LOOPER_ID_MAIN, ALOOPER_EVENT_INPUT, NULL,
@@ -1068,6 +1077,8 @@ static void* AndroidEventThreadWorker( void* param )
 	while (!IsEngineExitRequested())
 	{
 		//FPlatformMisc::LowLevelOutputDebugString(TEXT("AndroidEventThreadWorker"));
+		extern uint64 GFrameCounter;
+		STANDALONE_DEBUG_LOGf(LogAndroid, TEXT("AndroidEventThreadWorker(started), EventThreadID=%d, GFrameCounter=%u"), EventThreadID, GFrameCounter);
 
 		AndroidProcessEvents(state);
 
@@ -1114,7 +1125,7 @@ void android_main(struct android_app* state)
 	GGameThreadId = FPlatformTLS::GetCurrentThreadId();
 
 	BootTimingPoint("android_main");
-	STANDALONE_DEBUG_LOGf(LogAndroid, TEXT("Entering native app glue main function"));
+	STANDALONE_DEBUG_LOGf(LogAndroid, TEXT("Entering native app glue main function, GGameThreadId=%d"), GGameThreadId);
 	
 	GNativeAndroidApp = state;
 	check(GNativeAndroidApp);
@@ -1581,9 +1592,7 @@ static void ActivateApp_EventThread()
 	{
 		// Seems this can occur.
 		STANDALONE_DEBUG_LOGf(LogAndroid, TEXT("UNEXPECTED -- event thread, activate app, ALREADY have HW window lock according to bAppIsActive_EventThread. bReadyToProcessEvents=%d, bAppIsActive_EventThread=%d"), bReadyToProcessEvents, bAppIsActive_EventThread);
-#if !USE_ANDROID_STANDALONE
 		return;
-#endif
 	}
 
 	// Unlock window when we're ready.
@@ -1610,7 +1619,10 @@ static void ActivateApp_EventThread()
 	}
 
 	FThreadHeartBeat::Get().ResumeHeartBeat(true);
+
+#if !USE_ANDROID_STANDALONE
 	FPreLoadScreenManager::EnableRendering(true);
+#endif
 
 	extern void AndroidThunkCpp_ShowHiddenAlertDialog();
 	AndroidThunkCpp_ShowHiddenAlertDialog();	
@@ -1621,9 +1633,9 @@ extern void BlockRendering();
 static void SuspendApp_EventThread()
 {
 	DEVELOPER_LOG_COMMANDCB_CASE(SuspendApp_EventThread);
-	if (!bAppIsActive_EventThread)
+	if (!bAppIsActive_EventThread || !bHasFocus)
 	{
-		STANDALONE_DEBUG_LOGf(LogAndroid, TEXT("SuspendApp_EventThread -> UNEXPECTED: event thread, SuspendApp_EventThread called but aborted, GAndroidWindowLockRefCount=%d, bReadyToProcessEvents=%d, bAppIsActive_EventThread=%d, IsInGameThread=%d"), GAndroidWindowLockRefCount, bReadyToProcessEvents, bAppIsActive_EventThread, IsInGameThread());
+		STANDALONE_DEBUG_LOGf(LogAndroid, TEXT("SuspendApp_EventThread -> UNEXPECTED: event thread, SuspendApp_EventThread called but aborted, GAndroidWindowLockRefCount=%d, bReadyToProcessEvents=%d, bAppIsActive_EventThread=%d, bHasFocus=%d"), GAndroidWindowLockRefCount, bReadyToProcessEvents, bAppIsActive_EventThread, bHasFocus);
 		return;
 	}
 	bAppIsActive_EventThread = false;
@@ -1631,7 +1643,7 @@ static void SuspendApp_EventThread()
 
 	UE_LOG(LogAndroid, Log, TEXT("SuspendApp_EventThread -> event thread, suspending app, acquiring HW window lock."));
 	GAndroidWindowLock_Lock("SuspendApp_EventThread");
-	STANDALONE_DEBUG_LOGf(LogAndroid, TEXT("SuspendApp_EventThread -> event thread, suspending app, acquiring HW window lock. bReadyToProcessEvents=%d, bAppIsActive_EventThread=%d, IsInGameThread=%d"), bReadyToProcessEvents, bAppIsActive_EventThread, IsInGameThread());
+	STANDALONE_DEBUG_LOGf(LogAndroid, TEXT("SuspendApp_EventThread -> event thread, suspending app, acquiring HW window lock. bReadyToProcessEvents=%d, bAppIsActive_EventThread=%d"), bReadyToProcessEvents, bAppIsActive_EventThread);
 
 	if (bReadyToProcessEvents == false)
 	{
@@ -1858,20 +1870,18 @@ static void OnAppCommandCB(struct android_app* app, int32_t cmd)
 	
 		break;
 	case APP_CMD_RESUME:
+
+		DEVELOPER_LOG_COMMANDCB_CASE(APP_CMD_RESUME);
+
 		/**
 		 * Command from main thread: the app's activity has been resumed.
 		 */
 		bIsResumed = true;
 
 		// assume focus on resume
-		bDidGainFocus = true;
 		bHasFocus = true;
 		DEVELOPER_LOG_COMMANDCB_CASE(APP_CMD_RESUME);
 
-		if (bHasWindow && bHasFocus && bIsResumed)
-		{
-			ActivateApp_EventThread();
-		}
 		STANDALONE_DEBUG_LOGf(LogAndroid, TEXT("Case APP_CMD_RESUME, bShouldRestartFromInterrupt=%d"), bShouldRestartFromInterrupt);
 		UE_LOG(LogAndroid, Log, TEXT("Case APP_CMD_RESUME"));
 		FAppEventManager::GetInstance()->EnqueueAppEvent(APP_EVENT_STATE_ON_RESUME);
@@ -1879,6 +1889,11 @@ static void OnAppCommandCB(struct android_app* app, int32_t cmd)
 		// trigger focus
 		FAppEventManager::GetInstance()->EnqueueAppEvent(APP_EVENT_STATE_WINDOW_GAINED_FOCUS);
 
+		if (bHasWindow && bHasFocus && bIsResumed)
+		{
+			ActivateApp_EventThread();
+		}
+
 		/*
 		* On the initial loading the restart method must be called immediately
 		* in order to restart the app if the startup movie was playing
@@ -1921,16 +1936,6 @@ static void OnAppCommandCB(struct android_app* app, int32_t cmd)
 		}
 #endif
 
-		// Restart on resuming if did not complete engine initialization
-		if (!bDidCompleteEngineInit && bDidGainFocus  && !bIgnorePauseOnDownloaderStart && bAllowReboot)
-		{
-// 			// only do this if early startup enabled
-// 			FString *EarlyRestart = FAndroidMisc::GetConfigRulesVariable(TEXT("earlyrestart"));
-// 			if (EarlyRestart != NULL && EarlyRestart->Equals("true", ESearchCase::IgnoreCase))
-// 			{
-// 				bShouldRestartFromInterrupt = true;
-// 			}
-		}
 		bIgnorePauseOnDownloaderStart = false;
 
 		/*
@@ -2301,7 +2306,7 @@ JNI_METHOD void Java_com_epicgames_makeaar_GameActivityForMakeAAR_nativeMain(JNI
 	pthread_attr_t otherAttr;
 	pthread_attr_init(&otherAttr);
 	pthread_attr_setdetachstate(&otherAttr, PTHREAD_CREATE_DETACHED);
-	pthread_create(&G_AndroidMainThread, &otherAttr, android_main, &dummy_state);
+	pthread_create(&G_AndroidMainThread, &otherAttr, AndroidMain, &dummy_state);
 
 	STANDALONE_DEBUG_LOGf(LogAndroid, TEXT("Created main thread"));
 }
@@ -2421,13 +2426,15 @@ JNI_METHOD bool Java_com_epicgames_makeaar_Engine_nativeMountPak(JNIEnv* jenv, j
 
 JNI_METHOD void Java_com_epicgames_makeaar_GameActivityForMakeAAR_nativeAppCommand(JNIEnv* jenv, jobject thiz, jint cmd)
 {
-	STANDALONE_DEBUG_LOGf(LogAndroid, TEXT("MAKEAAR -> nativeAppCommand[%u, %s] -> EventThreadID=%d, GResumeMainInit=%d, GEventHandlerInitialized=%d, FTaskGraphInterface::IsRunning()=%d"), 
+	STANDALONE_DEBUG_LOGf(LogAndroid, TEXT("MAKEAAR -> nativeAppCommand[%u, %s] -> EventThreadID=%d, GGameThreadId=%d, GResumeMainInit=%d, GEventHandlerInitialized=%d, FTaskGraphInterface::IsRunning()=%d, FPlatformTLS::GetCurrentThreadId()=%d"), 
 		cmd,
 		get_APP_CMD_String(cmd),
 		EventThreadID, 
+		GGameThreadId,
 		GResumeMainInit,
 		GEventHandlerInitialized,
-		FTaskGraphInterface::IsRunning()
+		FTaskGraphInterface::IsRunning(),
+		FPlatformTLS::GetCurrentThreadId()
 	);
 
 	dummy_state.window = (ANativeWindow*)GAndroidWindowOverride;
