EGL's Robust context can change the result of GL's shader compiler/linker.
Make sure the robust context setting is replicated on the external compilers services and also changes OGL shader cache hash when enabled.

diff --git a/Engine/Build/Android/Java/src/com/epicgames/unreal/psoservices/PSOProgramService.java b/Engine/Build/Android/Java/src/com/epicgames/unreal/psoservices/PSOProgramService.java
--- a/Engine/Build/Android/Java/src/com/epicgames/unreal/psoservices/PSOProgramService.java
+++ b/Engine/Build/Android/Java/src/com/epicgames/unreal/psoservices/PSOProgramService.java
@@ -63,12 +63,15 @@ public class PSOProgramService extends Service implements Logger.ILoggerCallback
 	static final String CompiledProgram_Key = "cpg";
 	static final String SHMem_Key = "shm";
 	static final String JobFail = "f";
+	static final String RobustContextKey = "rbc";
 	public static final String LogDir = "/oglservice/";
 	public static final String LogExt = ".txt";
 	public static final Level LogLevel = Level.WARNING;
 
 	public static final boolean bEnableTestFailures = false;
 	public static final boolean bReportMemUse = false;
+	public boolean bWantRobustContext = false;
+	public boolean bGFXInitialized = false;
 	private FileHandler logFileHandler;
 
 	private void PrepareLogger() throws IOException
@@ -265,6 +268,33 @@ public class PSOProgramService extends Service implements Logger.ILoggerCallback
 	public IBinder onBind(Intent intent)
 	{
 		logger.debug("onBind "  + intent.toString());
+
+		if(!bGFXInitialized)
+		{
+			if(!UseVulkan())
+			{
+				Bundle extras =intent.getExtras();
+				if(extras != null)
+				{
+					bWantRobustContext = extras.getBoolean(RobustContextKey, false);
+					logger.debug("robust "  + bWantRobustContext);
+				}
+				else
+				{
+					logger.debug("no robust set");
+				}
+				logger.verbose("initGLContext " );
+				initGLContext();
+			}
+			else
+			{
+				logger.verbose("initVulkanContext " );
+				initVulkanContext();
+			}
+			bGFXInitialized = true;
+		}
+		LogMemInfo();
+
 		return mMessenger.getBinder();
 	}
 
@@ -444,17 +474,6 @@ public class PSOProgramService extends Service implements Logger.ILoggerCallback
 
 		Logger.RegisterCallback(this);
 		logger.verbose("oncreate " );
-		if(!UseVulkan())
-		{
-			logger.verbose("initGLContext " );
-			initGLContext();
-		}
-		else
-		{
-			logger.verbose("initVulkanContext " );
-			initVulkanContext();
-		}
-		LogMemInfo();
 	}
 
 	@Override
@@ -620,14 +639,28 @@ public class PSOProgramService extends Service implements Logger.ILoggerCallback
 			};
 		EGL14.eglChooseConfig(mEglDisplay, configSpec, 0, configs, 0, 1, num_config, 0);
 		//logger.verbose("2 eglChooseConfig "+num_config[0]+" err "+EGL14.eglGetError()+" configs: "+configs[0].toString());
-		int[] contextAttribsES31 = new int[]
+		int[] contextAttribsES32 = new int[]
+			{
+				EGLExt.EGL_CONTEXT_MAJOR_VERSION_KHR, 3,
+				EGLExt.EGL_CONTEXT_MINOR_VERSION_KHR, 2,
+				EGL14.EGL_NONE
+			};
+
+		int EGL_CONTEXT_OPENGL_ROBUST_ACCESS_EXT = 0x30BF;
+		int[] contextAttribsES32Robust = new int[]
 			{
 				EGLExt.EGL_CONTEXT_MAJOR_VERSION_KHR, 3,
 				EGLExt.EGL_CONTEXT_MINOR_VERSION_KHR, 2,
+				EGL_CONTEXT_OPENGL_ROBUST_ACCESS_EXT, 1,
 				EGL14.EGL_NONE
 			};
-		mEglContext = EGL14.eglCreateContext(mEglDisplay, configs[0], EGL14.EGL_NO_CONTEXT, contextAttribsES31, 0);
-		//logger.verbose("2 eglCreateContext "+mEglContext+" err "+EGL14.eglGetError());
+
+		String eglExtensions = EGL14.eglQueryString(mEglDisplay, EGL14.EGL_EXTENSIONS);
+		boolean bUseRobustContext = bWantRobustContext && eglExtensions.contains("EGL_EXT_create_context_robustness");
+		logger.verbose("2 eglCreateContext rbst "+bWantRobustContext+" ext contains "+eglExtensions.contains("EGL_EXT_create_context_robustness") );
+
+		mEglContext = EGL14.eglCreateContext(mEglDisplay, configs[0], EGL14.EGL_NO_CONTEXT, bUseRobustContext ? contextAttribsES32Robust : contextAttribsES32, 0);
+		logger.verbose("2 eglCreateContext "+mEglContext+", robust "+bUseRobustContext+", err "+EGL14.eglGetError());
 
 		if (EGL14.eglQueryString(mEglDisplay, EGL14.EGL_EXTENSIONS).contains("EGL_KHR_surfaceless_context"))
 		{
diff --git a/Engine/Build/Android/Java/src/com/epicgames/unreal/psoservices/PSOProgramServiceAccessor.java b/Engine/Build/Android/Java/src/com/epicgames/unreal/psoservices/PSOProgramServiceAccessor.java
--- a/Engine/Build/Android/Java/src/com/epicgames/unreal/psoservices/PSOProgramServiceAccessor.java
+++ b/Engine/Build/Android/Java/src/com/epicgames/unreal/psoservices/PSOProgramServiceAccessor.java
@@ -207,7 +207,7 @@ public class PSOProgramServiceAccessor
 		}
 	}
 
-	public static boolean AndroidThunkJava_StartRemoteProgramLink(int numServices, boolean bUseVulkan)
+	public static boolean AndroidThunkJava_StartRemoteProgramLink(int numServices, boolean bUseRobustEGLContext, boolean bUseVulkan)
 	{
 		boolean bSuccess = false;
 		if( _PSOProgramServiceAccessor != null)
@@ -225,7 +225,7 @@ public class PSOProgramServiceAccessor
 				}
 				else
 				{
-					bSuccess = _PSOProgramServiceAccessor.StartAndWaitForServices(numServices);
+					bSuccess = _PSOProgramServiceAccessor.StartAndWaitForServices(numServices, bUseRobustEGLContext);
 				}
 			}
 			catch (Exception e)
@@ -267,14 +267,14 @@ public class PSOProgramServiceAccessor
 		}
 	}
 
-	boolean StartAndWaitForServices(int numServices)
+	boolean StartAndWaitForServices(int numServices, boolean bUseRobustEGLContext)
 	{
 		numServices = Math.max(1, Math.min(numServices, ServiceClassTypes.length));
 
 		ServiceInstances = new OGLServiceInstance[numServices];
 		for(int i = 0; i< numServices ;i++)
 		{
-			ServiceInstances[i] = new OGLServiceInstance(ServiceClassTypes[i]);
+			ServiceInstances[i] = new OGLServiceInstance(ServiceClassTypes[i], bUseRobustEGLContext);
 		}
 
 		boolean bSuccess = true;
@@ -292,7 +292,7 @@ public class PSOProgramServiceAccessor
 		ServiceInstances = new OGLServiceInstance[numServices];
 		for(int i = 0; i< numServices ;i++)
 		{
-			ServiceInstances[i] = new OGLServiceInstance(VulkanServiceClassTypes[i]);
+			ServiceInstances[i] = new OGLServiceInstance(VulkanServiceClassTypes[i], false);
 		}
 
 		boolean bSuccess = true;
@@ -857,6 +857,7 @@ public class PSOProgramServiceAccessor
 	class OGLServiceInstance
 	{
 		private final Class ServiceClass;
+		private final boolean bRobustContext;
 		boolean mShouldUnbind = false;
 		private final AtomicInteger mBound = new AtomicInteger(-1);
 		Messenger mService = null;
@@ -867,9 +868,10 @@ public class PSOProgramServiceAccessor
 			return ServiceClass.getSimpleName();
 		}
 
-		public OGLServiceInstance(Class ServiceClassIN)
+		public OGLServiceInstance(Class ServiceClassIN, boolean bRobustContextIN)
 		{
 			ServiceClass = ServiceClassIN;
+			bRobustContext = bRobustContextIN;
 		}
 
 		private final Map<String, Long> LastLogSequencePerUID = new HashMap<>();
@@ -984,10 +986,12 @@ public class PSOProgramServiceAccessor
 				{
 					BindServiceFlags = BindServiceFlags | Context.BIND_IMPORTANT;
 				}
+				 
+				intent.putExtra(PSOProgramService.RobustContextKey,bRobustContext);
 
 				mShouldUnbind = mContext.bindService(intent, mConnection, BindServiceFlags);
 
-				Log.verbose("doBindService " + Name() + " needs unbind " + mShouldUnbind + " bound: " + mBound);
+				Log.verbose("doBindService " + Name() + " needs unbind " + mShouldUnbind + " bound: " + mBound+ "robust: "+bRobustContext);
 			}
 			return mShouldUnbind;
 		}
diff --git a/Engine/Source/Runtime/OpenGLDrv/Private/Android/AndroidEGL.cpp b/Engine/Source/Runtime/OpenGLDrv/Private/Android/AndroidEGL.cpp
--- a/Engine/Source/Runtime/OpenGLDrv/Private/Android/AndroidEGL.cpp
+++ b/Engine/Source/Runtime/OpenGLDrv/Private/Android/AndroidEGL.cpp
@@ -765,12 +765,22 @@ void AndroidEGL::ReInit()
 	SetCurrentSharedContext();
 }
 
+
+// EGL_CONTEXT_OPENGL_ROBUST_ACCESS_EXT is enabled if configrules asks for it or the command line specifies it.
+// If -OpenGLRobustContext=[0/1] is specified on the command line it takes precedence.
 void AndroidEGL::Init(APIVariant API, uint32 MajorVersion, uint32 MinorVersion)
 {
 	check(IsInGameThread());
 	const bool bDebug = IsOGLDebugOutputEnabled();
 	const FString* ConfigRulesForceRobustGLContext = FAndroidMisc::GetConfigRulesVariable(TEXT("ForceRobustGLContext"));
 	bool bWantsRobustGLContext = ConfigRulesForceRobustGLContext && ConfigRulesForceRobustGLContext->Equals("true", ESearchCase::IgnoreCase);
+	
+	FString RobustArg;
+	if (FParse::Value(FCommandLine::Get(), TEXT("-OpenGLRobustContext="), RobustArg))
+	{
+		bWantsRobustGLContext = RobustArg.Contains(TEXT("1"));
+	}
+	
 	if (PImplData->Initalized)
 	{
 		ensure(bDebug == PImplData->bIsDebug); // if this fires you would need to tear down the previous context and recreate to honour the debug change.
@@ -801,7 +811,8 @@ void AndroidEGL::Init(APIVariant API, uint32 MajorVersion, uint32 MinorVersion)
 		}
 #endif // USE_ANDROID_EGL_NO_ERROR_CONTEXT
 
-		if (bSupportsEXTRobustContext && bWantsRobustGLContext)
+		bIsEXTRobustContextActive = bSupportsEXTRobustContext && bWantsRobustGLContext;
+		if (bIsEXTRobustContextActive)
 		{
 			UE_LOG(LogAndroid, Log, TEXT("Enabling: EGL_CONTEXT_OPENGL_ROBUST_ACCESS_EXT"));
 			ContextAttributes[Element++] = EGL_CONTEXT_OPENGL_ROBUST_ACCESS_EXT;
diff --git a/Engine/Source/Runtime/OpenGLDrv/Private/Android/AndroidEGL.h b/Engine/Source/Runtime/OpenGLDrv/Private/Android/AndroidEGL.h
--- a/Engine/Source/Runtime/OpenGLDrv/Private/Android/AndroidEGL.h
+++ b/Engine/Source/Runtime/OpenGLDrv/Private/Android/AndroidEGL.h
@@ -85,7 +85,8 @@ public:
 	bool IsOfflineSurfaceRequired();
 
 	void GetDimensions(uint32& OutWidth, uint32& OutHeight);
-	
+	bool IsUsingRobustContext() const { return bIsEXTRobustContextActive; }
+
 	EGLDisplay GetDisplay() const;
 	EGLSurface GetSurface() const;
 	EGLConfig GetConfig() const;
@@ -146,6 +147,7 @@ private:
 	bool bSupportsKHRSurfacelessContext;
 	bool bSupportsKHRNoErrorContext;
 	bool bSupportsEXTRobustContext;
+	bool bIsEXTRobustContextActive = false;
 
 	int *ContextAttributes;
 };
diff --git a/Engine/Source/Runtime/OpenGLDrv/Private/Android/AndroidOpenGL.cpp b/Engine/Source/Runtime/OpenGLDrv/Private/Android/AndroidOpenGL.cpp
--- a/Engine/Source/Runtime/OpenGLDrv/Private/Android/AndroidOpenGL.cpp
+++ b/Engine/Source/Runtime/OpenGLDrv/Private/Android/AndroidOpenGL.cpp
@@ -130,7 +130,7 @@ struct FOpenGLRemoteGLProgramCompileJNI
 		{
 			DispatchProgramLink = FJavaWrapper::FindStaticMethod(Env, OGLServiceAccessor, "AndroidThunkJava_OGLRemoteProgramLink", "([BLjava/lang/String;Ljava/lang/String;Ljava/lang/String;)Lcom/epicgames/unreal/psoservices/PSOProgramServiceAccessor$JNIProgramLinkResponse;", false);
 			CHECK_JNI_EXCEPTIONS(Env);
-			StartRemoteProgramLink = FJavaWrapper::FindStaticMethod(Env, OGLServiceAccessor, "AndroidThunkJava_StartRemoteProgramLink", "(IZ)Z", false);
+			StartRemoteProgramLink = FJavaWrapper::FindStaticMethod(Env, OGLServiceAccessor, "AndroidThunkJava_StartRemoteProgramLink", "(IZZ)Z", false);
 			CHECK_JNI_EXCEPTIONS(Env);
 			StopRemoteProgramLink = FJavaWrapper::FindStaticMethod(Env, OGLServiceAccessor, "AndroidThunkJava_StopRemoteProgramLink", "()V", false);
 			CHECK_JNI_EXCEPTIONS(Env);
@@ -1136,7 +1136,8 @@ bool FAndroidOpenGL::StartAndWaitForRemoteCompileServices(int NumServices)
 
 	if (Env && AreAndroidOpenGLRemoteCompileServicesAvailable())
 	{
-		bResult = (bool)Env->CallStaticBooleanMethod(OpenGLRemoteGLProgramCompileJNI.OGLServiceAccessor, OpenGLRemoteGLProgramCompileJNI.StartRemoteProgramLink, (jint)NumServices, (jboolean)false);
+		bool bUseRobustContexts = AndroidEGL::GetInstance()->IsUsingRobustContext();
+		bResult = (bool)Env->CallStaticBooleanMethod(OpenGLRemoteGLProgramCompileJNI.OGLServiceAccessor, OpenGLRemoteGLProgramCompileJNI.StartRemoteProgramLink, (jint)NumServices, (jboolean)bUseRobustContexts, /*bUseVulkan*/(jboolean)false);
 		GRemoteCompileServicesActive = bResult;
 	}
 
diff --git a/Engine/Source/Runtime/OpenGLDrv/Private/OpenGLProgramBinaryFileCache.cpp b/Engine/Source/Runtime/OpenGLDrv/Private/OpenGLProgramBinaryFileCache.cpp
--- a/Engine/Source/Runtime/OpenGLDrv/Private/OpenGLProgramBinaryFileCache.cpp
+++ b/Engine/Source/Runtime/OpenGLDrv/Private/OpenGLProgramBinaryFileCache.cpp
@@ -234,6 +234,8 @@ FOpenGLProgramBinaryCache::FOpenGLProgramBinaryCache(const FString& InCachePathR
 	// Some devices report binary compatibility errors after minor OS updates even though the GL driver version has not changed.
 	const FString BuildNumber = FAndroidMisc::GetDeviceBuildNumber();
 	HashString.Append(BuildNumber);
+	
+	HashString.Append(AndroidEGL::GetInstance()->IsUsingRobustContext() ? TEXT("ROBUST") : TEXT("NRB"));
 
 	// Optional configrule variable for triggering a rebuild of the cache.
 	const FString* ConfigRulesGLProgramKey = FAndroidMisc::GetConfigRulesVariable(TEXT("OpenGLProgramCacheKey"));
diff --git a/Engine/Source/Runtime/VulkanRHI/Private/Android/VulkanAndroidPlatform.cpp b/Engine/Source/Runtime/VulkanRHI/Private/Android/VulkanAndroidPlatform.cpp
--- a/Engine/Source/Runtime/VulkanRHI/Private/Android/VulkanAndroidPlatform.cpp
+++ b/Engine/Source/Runtime/VulkanRHI/Private/Android/VulkanAndroidPlatform.cpp
@@ -1388,7 +1388,7 @@ struct FVKRemoteProgramCompileJNI
 			CHECK_JNI_EXCEPTIONS(Env);
 			DispatchPSOCompileShm = FJavaWrapper::FindStaticMethod(Env, PSOServiceAccessor, "AndroidThunkJava_VKPSOGFXCompileShm", "([BIJJJJZ)Lcom/epicgames/unreal/psoservices/PSOProgramServiceAccessor$JNIProgramLinkResponse;", false);
 			CHECK_JNI_EXCEPTIONS(Env);
-			StartRemoteProgramLink = FJavaWrapper::FindStaticMethod(Env, PSOServiceAccessor, "AndroidThunkJava_StartRemoteProgramLink", "(IZ)Z", false);
+			StartRemoteProgramLink = FJavaWrapper::FindStaticMethod(Env, PSOServiceAccessor, "AndroidThunkJava_StartRemoteProgramLink", "(IZZ)Z", false);
 			CHECK_JNI_EXCEPTIONS(Env);
 			StopRemoteProgramLink = FJavaWrapper::FindStaticMethod(Env, PSOServiceAccessor, "AndroidThunkJava_StopRemoteProgramLink", "()V", false);
 			CHECK_JNI_EXCEPTIONS(Env);
@@ -1446,7 +1446,7 @@ bool FVulkanAndroidPlatform::StartAndWaitForRemoteCompileServices(int NumService
 
 	if (Env && AreAndroidVulkanRemoteCompileServicesAvailable())
 	{
-		bResult = (bool)Env->CallStaticBooleanMethod(VKRemoteProgramCompileJNI.PSOServiceAccessor, VKRemoteProgramCompileJNI.StartRemoteProgramLink, (jint)NumServices, (jboolean)true);
+		bResult = (bool)Env->CallStaticBooleanMethod(VKRemoteProgramCompileJNI.PSOServiceAccessor, VKRemoteProgramCompileJNI.StartRemoteProgramLink, (jint)NumServices, /*bUseRobustEGLContext*/(jboolean)false, /*bUseVulkan*/(jboolean)true);
 		GRemoteCompileServicesActive = bResult;
 	}
 
