Update Android styles to more modern base theme

#android

diff --git a/Engine/Build/Android/Java/res/values-land/styles.xml b/Engine/Build/Android/Java/res/values-land/styles.xml
--- a/Engine/Build/Android/Java/res/values-land/styles.xml
+++ b/Engine/Build/Android/Java/res/values-land/styles.xml
@@ -1,7 +1,16 @@
 <?xml version="1.0" encoding="utf-8"?>
 <resources>
 
-    <style name="UnrealBaseTheme" parent = "@android:style/Theme.Black.NoTitleBar.Fullscreen" />
+	<style name="Theme.Black.AppCompat.NoActionBar.FullScreen" parent="@style/Theme.AppCompat.NoActionBar">
+		<item name="android:windowNoTitle">true</item>
+		<item name="android:windowActionBar">false</item>
+		<item name="android:windowFullscreen">true</item>
+		<item name="android:windowContentOverlay">@null</item>
+		<item name="android:windowBackground">@android:color/black</item>
+		<item name="android:colorBackground">@android:color/black</item>
+	</style>
+
+    <style name="UnrealBaseTheme" parent = "Theme.Black.AppCompat.NoActionBar.FullScreen" />
 
     <style name="UnrealSplashTheme" parent = "UnrealBaseTheme">
         <item name="android:windowBackground">@drawable/splashscreen_l</item>
diff --git a/Engine/Build/Android/Java/res/values-port/styles.xml b/Engine/Build/Android/Java/res/values-port/styles.xml
--- a/Engine/Build/Android/Java/res/values-port/styles.xml
+++ b/Engine/Build/Android/Java/res/values-port/styles.xml
@@ -1,7 +1,16 @@
 <?xml version="1.0" encoding="utf-8"?>
 <resources>
 
-    <style name="UnrealBaseTheme" parent = "@android:style/Theme.Black.NoTitleBar.Fullscreen" />
+ 	<style name="Theme.Black.AppCompat.NoActionBar.FullScreen" parent="@style/Theme.AppCompat.NoActionBar">
+		<item name="android:windowNoTitle">true</item>
+		<item name="android:windowActionBar">false</item>
+		<item name="android:windowFullscreen">true</item>
+		<item name="android:windowContentOverlay">@null</item>
+		<item name="android:windowBackground">@android:color/black</item>
+		<item name="android:colorBackground">@android:color/black</item>
+	</style>
+
+    <style name="UnrealBaseTheme" parent = "Theme.Black.AppCompat.NoActionBar.FullScreen" />
 
     <style name="UnrealSplashTheme" parent = "UnrealBaseTheme">
         <item name="android:windowBackground">@drawable/splashscreen_p</item>
diff --git a/Engine/Build/Android/Java/res/values/styles.xml b/Engine/Build/Android/Java/res/values/styles.xml
--- a/Engine/Build/Android/Java/res/values/styles.xml
+++ b/Engine/Build/Android/Java/res/values/styles.xml
@@ -1,7 +1,16 @@
 <?xml version="1.0" encoding="utf-8"?>
 <resources>
 
-    <style name="UnrealBaseTheme" parent = "@android:style/Theme.Black.NoTitleBar.Fullscreen" />
+	<style name="Theme.Black.AppCompat.NoActionBar.FullScreen" parent="@style/Theme.AppCompat.NoActionBar">
+		<item name="android:windowNoTitle">true</item>
+		<item name="android:windowActionBar">false</item>
+		<item name="android:windowFullscreen">true</item>
+		<item name="android:windowContentOverlay">@null</item>
+		<item name="android:windowBackground">@android:color/black</item>
+		<item name="android:colorBackground">@android:color/black</item>
+	</style>
+
+    <style name="UnrealBaseTheme" parent = "Theme.Black.AppCompat.NoActionBar.FullScreen" />
 
     <style name="UnrealSplashTheme" parent = "UnrealBaseTheme" />
 
