Set input text to black on Android with new theme

#android

diff --git a/Engine/Build/Android/Java/src/com/epicgames/unreal/GameActivity.java.template b/Engine/Build/Android/Java/src/com/epicgames/unreal/GameActivity.java.template
--- a/Engine/Build/Android/Java/src/com/epicgames/unreal/GameActivity.java.template
+++ b/Engine/Build/Android/Java/src/com/epicgames/unreal/GameActivity.java.template
@@ -6481,6 +6481,7 @@ public class GameActivity extends $${gameActivitySuperClass}$$ implements Surfac
 		newVirtualKeyboardInput = new VirtualKeyboardInput(this);
 		newVirtualKeyboardInput.setSingleLine(false);
 		newVirtualKeyboardInput.setBackgroundColor(Color.WHITE);
+		newVirtualKeyboardInput.setTextColor(Color.BLACK);
 		newVirtualKeyboardInput.setLayoutParams(new RelativeLayout.LayoutParams(LayoutParams.MATCH_PARENT, LayoutParams.WRAP_CONTENT));
 		newVirtualKeyboardInput.setVisibility(View.GONE);
 		if (ANDROID_BUILD_VERSION < 11)
