New: Cleanup and simplify GameActivity.java.template changes to be less intrusive and support ASIS.
BugFix: Cleanup Android engine source logic to have minimal impact when not using ASIS.

diff --git a/Engine/Build/Android/Java/src/com/epicgames/unreal/GameActivity.java.template b/Engine/Build/Android/Java/src/com/epicgames/unreal/GameActivity.java.template
--- a/Engine/Build/Android/Java/src/com/epicgames/unreal/GameActivity.java.template
+++ b/Engine/Build/Android/Java/src/com/epicgames/unreal/GameActivity.java.template
@@ -4,6 +4,7 @@
 package com.epicgames.unreal;
 
 import android.app.NotificationManager;
+import android.content.ContextWrapper;
 import android.opengl.EGL14;
 import android.opengl.EGLConfig;
 import android.opengl.EGLContext;
@@ -11,6 +12,7 @@ import android.opengl.EGLDisplay;
 import android.opengl.EGLSurface;
 import android.opengl.GLES20;
 
+import java.lang.annotation.Native;
 import java.util.regex.Pattern;
 import java.util.regex.Matcher;
 
@@ -340,7 +342,7 @@ public class GameActivity extends $${gameActivitySuperClass}$$ implements Surfac
 	
 	protected boolean isStandalone = false;
 	protected static GameActivity _gameActivityInstance = null;
-	protected SimpleContextWrapper activityContext = null;
+	protected ContextWrapper activityContext = null;
 	protected static Bundle _bundle;
 	protected static Bundle _extrasBundle;
 	protected boolean bOnCreateCalled = false;
@@ -559,19 +561,83 @@ public class GameActivity extends $${gameActivitySuperClass}$$ implements Surfac
 		return _gameActivityInstance;
 	}
 
+	public static Boolean isStandaloneMode()
+	{
+		return _gameActivityInstance.isStandalone;
+	}
+
+	
 	public static Boolean isValidGameActivity()
 	{
 		return _gameActivityInstance != null;
 	}
 
-	public static SimpleContextWrapper GetCurrentActivityContext()
+	public static ContextWrapper GetCurrentActivityContext()
 	{
 		return _gameActivityInstance.activityContext;
 	}
 
+	public static Activity getCurrentGameActivity()
+	{
+		GameActivity ga = Get();
+		if (ga.isStandalone)
+		{
+			assert (ga._gameActivityInstance.activityContext instanceof Activity);
+			return (Activity)ga._gameActivityInstance.activityContext;
+		}
+		else
+		{
+			assert (_gameActivityInstance.getBaseContext() != null);
+			return ga;
+		}
+	}
+
+	public final void runOnMainThread(Runnable action) {
+		assert (getBaseContext() != null);
+		assert (getMainExecutor() != null);
+		getMainExecutor().execute(action);
+	}
+	
+	private Activity getGameActivity()
+	{
+		assert(activityContext != null);
+		if (isStandalone)
+		{
+			if (activityContext instanceof Activity) 
+			{
+				return (Activity) activityContext;
+			}
+			else
+			{
+				Context searchContext = activityContext;
+				Log.debug( "LIFECYCLE: GameActivity.getGameActivity() searching for Activity in _gameActivityInstance =" + _gameActivityInstance + ", activityContext=" + activityContext);
+
+				while (searchContext instanceof ContextWrapper) {
+					if (searchContext instanceof Activity) {
+						Log.warn( "LIFECYCLE: GameActivity.getGameActivity() updating activityContext = " + activityContext + ", with searchContext=" + (ContextWrapper)searchContext);
+
+						activityContext = (ContextWrapper) searchContext;
+						return (Activity) searchContext;
+					}
+					searchContext = ((ContextWrapper) searchContext).getBaseContext();
+				}
+
+				Log.error( "LIFECYCLE ERROR: GameActivity.getGameActivity() did not find Activity for _gameActivityInstance =" + _gameActivityInstance + ", activityContext=" + activityContext);
+
+				assert (activityContext instanceof Activity);
+				return null;
+			}
+		}
+		else
+		{
+			assert (getBaseContext() != null);
+			return this;
+		}
+	}
+	
 	public final Window getWindow() 
 	{
-		return isStandalone ? activityContext.getWindow() : super.getWindow();
+		return isStandalone ? getGameActivity().getWindow() : super.getWindow();
 	}
 
 	public final Intent getIntent()
@@ -581,7 +647,14 @@ public class GameActivity extends $${gameActivitySuperClass}$$ implements Surfac
 
 	public final WindowManager getWindowManager()
 	{
-		return isStandalone ? activityContext.getWindowManager() : super.getWindowManager();
+		if (isStandalone) 
+		{
+			return (WindowManager) activityContext.getSystemService(Context.WINDOW_SERVICE);
+		}
+		else 
+		{
+			return super.getWindowManager();
+		}
 	}
 
 	
@@ -738,7 +811,7 @@ public class GameActivity extends $${gameActivitySuperClass}$$ implements Surfac
 	public String AndroidThunkJava_GetCommandLine()
 	{
 		String commandLine = "";
-		//$${gameActivityGetCommandLineAdditions}$$
+//$${gameActivityGetCommandLineAdditions}$$
 		return commandLine;
 	}
 
@@ -2435,10 +2508,12 @@ public class GameActivity extends $${gameActivitySuperClass}$$ implements Surfac
 	{
 		assert(_gameActivityInstance != null);
 		bOnCreateCalled = true;
-		activityContext = new SimpleContextWrapper(this);
+		activityContext = (ContextWrapper) getApplicationContext();
+		assert (activityContext != null);
+		
 		InternalFilesDir = getFilesDir().getAbsolutePath() + "/";
 		ExternalFilesDir = getExternalFilesDir(null).getAbsolutePath() + "/";
-		//$${gameActivityOnCreateBeginningAdditions}$$
+//$${gameActivityOnCreateBeginningAdditions}$$
 		Logger.RegisterCallback(this);
 
 		// Grab a reference to the asset manager
@@ -2451,6 +2526,47 @@ public class GameActivity extends $${gameActivitySuperClass}$$ implements Surfac
 
 	protected void onCreateBody(Bundle savedInstanceState)
 	{
+		if (super.getBaseContext() == null)
+		{
+			if (isStandalone)
+			{
+				Log.debug("LIFECYCLE: onCreateBody(): super.getBaseContext() is null so attach base context = " + activityContext);
+			}
+			attachBaseContext(activityContext);
+		}
+		if (isStandalone)
+		{
+			Log.debug("LIFECYCLE: onCreateBody(): isStandalone=" + isStandalone 
+				+ ", getBaseContext() = " + getBaseContext() 
+				+ ", activityContext = "+ activityContext);
+		}
+		
+		if (super.getApplicationContext() == null)
+		{
+			if (isStandalone)
+			{
+				Log.debug("LIFECYCLE: onCreateBody(): super.getApplicationContext() is null so attach base context = " + activityContext);
+			}
+			attachBaseContext(activityContext);
+		}
+		
+		GameApplication gameApplication = GameApplication.getAppContext() instanceof GameApplication ? (GameApplication) GameApplication.getAppContext() : null;
+		if (gameApplication != null && gameApplication.getBaseContext() == null)
+		{
+			if (isStandalone)
+			{
+				Log.debug("LIFECYCLE: onCreateBody(): attachBaseContext for GameApplication to context = " + activityContext);
+			}
+			gameApplication.attachBaseContext(activityContext);
+		}
+
+		if (isStandalone)
+		{
+			Log.debug("LIFECYCLE: onCreateBody(): getBaseContext() = " + getBaseContext());
+			Log.debug("LIFECYCLE: onCreateBody(): GameApplication.getAppContext() = " + GameApplication.getAppContext());
+			Log.debug("LIFECYCLE: onCreateBody(): activityContext.getApplicationContext() = " + activityContext.getApplicationContext());
+		}
+
 		sensorManager = (SensorManager) getSystemService(Context.SENSOR_SERVICE);
 		accelerometer = sensorManager.getDefaultSensor(Sensor.TYPE_ACCELEROMETER);
 		magnetometer = sensorManager.getDefaultSensor(Sensor.TYPE_MAGNETIC_FIELD);
@@ -2476,7 +2592,7 @@ public class GameActivity extends $${gameActivitySuperClass}$$ implements Surfac
 		}
 
 		// disable autofill to avoid possible Android 10 crash (android.view.autofill.AutofillManager.ensureServiceClientAddedIfNeededLocked() SyncResultReceiver$TimeoutException)
-		if (Build.VERSION.SDK_INT >= Build.VERSION_CODES.O)
+		if (!isStandalone && Build.VERSION.SDK_INT >= Build.VERSION_CODES.O)
 		{
 			try
 			{
@@ -2550,7 +2666,7 @@ public class GameActivity extends $${gameActivitySuperClass}$$ implements Surfac
 				CommandLineFull = CmdLine;
 			}
 
-			if (UseDisplayCutout)
+			if (!isStandalone && UseDisplayCutout)
 			{
 				// will not be true if not Android Pie or later
 	            WindowManager.LayoutParams params = getWindow().getAttributes();
@@ -2653,12 +2769,11 @@ public class GameActivity extends $${gameActivitySuperClass}$$ implements Surfac
 							memLog += "memstat: "+memstatusline + "\n";
 				    	}
 					}
-					//$${gameActivityMemStatAdditions}$$
+//$${gameActivityMemStatAdditions}$$
 					Log.debug(memLog);
 				}
 				synchronized(activityContext)
 				{
-					Activity localActivity = SimpleContextWrapper.getActivityOfContext(activityContext);
 					if (ProcessMemory > 0)
 					{
 						UsedMemory = ProcessMemory;
@@ -2683,7 +2798,7 @@ public class GameActivity extends $${gameActivitySuperClass}$$ implements Surfac
 			MarginLayoutParams params = new MarginLayoutParams(LayoutParams.WRAP_CONTENT, LayoutParams.WRAP_CONTENT);
 			params.setMargins(0, 0, 0, 0);
 			activityLayout = new LinearLayout(activityContext);
-			SimpleContextWrapper.getActivityOfContext(activityContext).setContentView(activityLayout, params);
+			getGameActivity().setContentView(activityLayout, params);
 
 			// tell Android that we want volume controls to change the media volume, aka music
 			setVolumeControlStream(AudioManager.STREAM_MUSIC);
@@ -2786,7 +2901,7 @@ public class GameActivity extends $${gameActivitySuperClass}$$ implements Surfac
 		}
 
 		// is this a native landscape device (tablet, tv)?
-		if ( getDeviceDefaultOrientation() == Configuration.ORIENTATION_LANDSCAPE )
+		if (!isStandalone && getDeviceDefaultOrientation() == Configuration.ORIENTATION_LANDSCAPE )
 		{
 			boolean bForceLandscape = false;
 
@@ -2834,7 +2949,7 @@ public class GameActivity extends $${gameActivitySuperClass}$$ implements Surfac
 			if (bForceLandscape)
 			{
 				Log.debug( "Setting screen orientation to landscape because we have detected landscape device" );
-				SimpleContextWrapper.getActivityOfContext(activityContext).setRequestedOrientation( android.content.pm.ActivityInfo.SCREEN_ORIENTATION_LANDSCAPE );
+				getGameActivity().setRequestedOrientation( android.content.pm.ActivityInfo.SCREEN_ORIENTATION_LANDSCAPE );
 			}
 		}
 
@@ -3405,7 +3520,7 @@ public class GameActivity extends $${gameActivitySuperClass}$$ implements Surfac
 
 	public void AndroidThunkJava_ResizeKeyboard(boolean bShow, Rect keyboardRect, Rect visibleRect)
 	{
-		//$${gameActivityResizeKeyboardAdditions}$$
+//$${gameActivityResizeKeyboardAdditions}$$
 	}
 
 	@Override
@@ -3431,7 +3546,7 @@ public class GameActivity extends $${gameActivitySuperClass}$$ implements Surfac
 
 	public void restoreTransparentBars()
 	{
-		if(android.os.Build.VERSION.SDK_INT >= 19) {
+		if(!isStandalone && android.os.Build.VERSION.SDK_INT >= 19) {
 			try {
 				View decorView = getWindow().getDecorView(); 
 
@@ -3635,7 +3750,7 @@ public class GameActivity extends $${gameActivitySuperClass}$$ implements Surfac
 		{
 			final boolean bWarning = bForceGameEndWithWarning;
 
-			activityContext.runOnUiThread(new Runnable()
+			runOnUiThread(new Runnable()
 			{
 				@Override
 				public void run()
@@ -3753,7 +3868,7 @@ public class GameActivity extends $${gameActivitySuperClass}$$ implements Surfac
 		{
 			//	If an AlertDialog is showing when the application is paused, it can cause our main window to be terminated
 			//	Hide the dialog here. It will be shown again via AndroidThunkJava_ShowHiddenAlertDialog called from native code
-			activityContext.runOnUiThread(new Runnable()
+			runOnUiThread(new Runnable()
 			{
 				public void run()
 				{
@@ -3995,7 +4110,7 @@ public class GameActivity extends $${gameActivitySuperClass}$$ implements Surfac
 	public void onConfigurationChanged(Configuration newConfig)
 	{
 		super.onConfigurationChanged(newConfig);
-		//$${gameActivityonConfigurationChangedAdditions}$$
+//$${gameActivityonConfigurationChangedAdditions}$$
 		// forward the orientation
 		nativeOnOrientationChanged(newConfig.orientation);
 	}
@@ -4026,7 +4141,7 @@ public class GameActivity extends $${gameActivitySuperClass}$$ implements Surfac
 			Log.debug("==============> [JAVA] AndroidThunkJava_ShowHiddenAlertDialog() - Showing " + CurrentDialogType);
 		
 			//	If an AlertDialog was showing onPause and we hid it, show it again
-			activityContext.runOnUiThread(new Runnable()
+			runOnUiThread(new Runnable()
 			{
 				public void run()
 				{
@@ -4062,7 +4177,7 @@ public class GameActivity extends $${gameActivitySuperClass}$$ implements Surfac
 		bKeepScreenOn = Enable;
 		if (Enable)
 		{
-			activityContext.runOnUiThread(new Runnable()
+			runOnUiThread(new Runnable()
 			{
 				@Override
 				public void run()
@@ -4074,7 +4189,7 @@ public class GameActivity extends $${gameActivitySuperClass}$$ implements Surfac
 		}
 		else
 		{
-			activityContext.runOnUiThread(new Runnable()
+			runOnUiThread(new Runnable()
 			{
 				@Override
 				public void run()
@@ -4100,7 +4215,7 @@ public class GameActivity extends $${gameActivitySuperClass}$$ implements Surfac
 
 		if (bDisable)
 		{
-			activityContext.runOnUiThread(new Runnable()
+			runOnUiThread(new Runnable()
 			{
 				@Override
 				public void run()
@@ -4112,7 +4227,7 @@ public class GameActivity extends $${gameActivitySuperClass}$$ implements Surfac
 		}
 		else
 		{
-			activityContext.runOnUiThread(new Runnable()
+			runOnUiThread(new Runnable()
 			{
 				@Override
 				public void run()
@@ -4161,7 +4276,7 @@ public class GameActivity extends $${gameActivitySuperClass}$$ implements Surfac
 	{
 		if (deviceVibrator != null)
 		{
-			activityContext.runOnUiThread(new VibrateRunnable(Intensity, (long)Duration, deviceVibrator));
+			runOnUiThread(new VibrateRunnable(Intensity, (long)Duration, deviceVibrator));
 		}
 	}
 
@@ -4178,7 +4293,7 @@ public class GameActivity extends $${gameActivitySuperClass}$$ implements Surfac
 		consoleInputBox.setHistoryEnd();
 
 		consoleAlert.setMessage("[Available texture formats: " + Formats + "]");
-		activityContext.runOnUiThread(new Runnable()
+		runOnUiThread(new Runnable()
 		{
 			public void run()
 			{
@@ -4201,7 +4316,7 @@ public class GameActivity extends $${gameActivitySuperClass}$$ implements Surfac
 			return;
 		}
 
-		activityContext.runOnUiThread(new Runnable()
+		runOnUiThread(new Runnable()
 		{
 			public void run()
 			{
@@ -4229,7 +4344,7 @@ public class GameActivity extends $${gameActivitySuperClass}$$ implements Surfac
 		final String uiLabel = inLabel;
 		final String uiContents = inContents;
 
-		activityContext.runOnUiThread(new Runnable()
+		runOnUiThread(new Runnable()
 		{
 			public void run()
 			{
@@ -4520,7 +4635,7 @@ public class GameActivity extends $${gameActivitySuperClass}$$ implements Surfac
 	public String AndroidThunkJava_GetFunnelId()
 	{
 		String funnelId = "";
-		//$${gameActivityGetFunnelIdAdditions}$$
+//$${gameActivityGetFunnelIdAdditions}$$
 		return funnelId;
 	}
 
@@ -4528,7 +4643,7 @@ public class GameActivity extends $${gameActivitySuperClass}$$ implements Surfac
 	{
 		String loginId = null;
 
-	//$${gameActivityGetLoginIdAdditions}$$
+//$${gameActivityGetLoginIdAdditions}$$
 
 		// check for existing GUID fallback in internal
 		File file = new File(InternalFilesDir + "login-identifier.txt");
@@ -4676,12 +4791,12 @@ public class GameActivity extends $${gameActivitySuperClass}$$ implements Surfac
 	public void AndroidThunkJava_SetOrientation(int value)
 	{
 		final int orientation = value;
-		activityContext.runOnUiThread(new Runnable()
+		runOnUiThread(new Runnable()
 		{
 			@Override
 			public void run()
 			{
-				SimpleContextWrapper.getActivityOfContext(activityContext).setRequestedOrientation( orientation );
+				getGameActivity().setRequestedOrientation( orientation );
 			}
 		});
 	}
@@ -4705,7 +4820,7 @@ public class GameActivity extends $${gameActivitySuperClass}$$ implements Surfac
 	// call back into native code from the Java UI thread, initializing any available VR HMD modules
 	public void AndroidThunkJava_InitHMDs()
 	{
-		activityContext.runOnUiThread(new Runnable()
+		runOnUiThread(new Runnable()
 		{
 			@Override
 			public void run()
@@ -4765,7 +4880,7 @@ public class GameActivity extends $${gameActivitySuperClass}$$ implements Surfac
 		{
 			bTriggeredQuery = true;
 
-			activityContext.runOnUiThread(new Runnable()
+			runOnUiThread(new Runnable()
 			{
 				@Override
 				public void run()
@@ -4932,7 +5047,7 @@ public class GameActivity extends $${gameActivitySuperClass}$$ implements Surfac
 		final String messageFinal = message;
 		final boolean bHorizontalFinal = bHorizontal;
 		final int maxValueFinal = maxValue;
-		activityContext.runOnUiThread(new Runnable()
+		runOnUiThread(new Runnable()
 		{
 			@Override
 			public void run()
@@ -4965,7 +5080,7 @@ public class GameActivity extends $${gameActivitySuperClass}$$ implements Surfac
 			return;
 		}
 		final int valueFinal = value;
-		activityContext.runOnUiThread(new Runnable()
+		runOnUiThread(new Runnable()
 		{
 			@Override
 			public void run()
@@ -5483,7 +5598,7 @@ public class GameActivity extends $${gameActivitySuperClass}$$ implements Surfac
 	{
 		if (vibrator != null)
 		{
-			activityContext.runOnUiThread(new VibrateRunnable(intensity, duration, vibrator));
+			runOnUiThread(new VibrateRunnable(intensity, duration, vibrator));
 			return true;
 		}
 		return false;
@@ -5538,7 +5653,7 @@ public class GameActivity extends $${gameActivitySuperClass}$$ implements Surfac
 			DesiredHolderHeight > 0 &&
 			MySurfaceView != null)
 		{
-			activityContext.runOnUiThread(new Runnable()
+			runOnUiThread(new Runnable()
 			{
 				@Override
 				public void run()
@@ -5569,7 +5684,7 @@ public class GameActivity extends $${gameActivitySuperClass}$$ implements Surfac
 
 		if(bUseSurfaceView && MySurfaceView != null)
 		{
-			activityContext.runOnUiThread(new Runnable()
+			runOnUiThread(new Runnable()
 			{
 				@Override
 				public void run()
@@ -5650,7 +5765,7 @@ public class GameActivity extends $${gameActivitySuperClass}$$ implements Surfac
 		if (ANDROID_BUILD_VERSION >= 30)
 		{
 			final int FinalRefreshRate = RefreshRate;
-			activityContext.runOnUiThread(new Runnable()
+			runOnUiThread(new Runnable()
 			{
 				@Override
 				public void run()
@@ -5689,7 +5804,7 @@ public class GameActivity extends $${gameActivitySuperClass}$$ implements Surfac
 					{
 						if (!isStandalone)
 						{
-							activityContext.runOnUiThread(new Runnable()
+							runOnUiThread(new Runnable()
 							{
 								public void run()
 								{
@@ -6197,7 +6312,7 @@ public class GameActivity extends $${gameActivitySuperClass}$$ implements Surfac
 			}
 			if (!isStandalone)
 			{
-				activityContext.runOnUiThread(new Runnable()
+				runOnUiThread(new Runnable()
 				{
 					public void run()
 					{
@@ -6799,7 +6914,7 @@ public class GameActivity extends $${gameActivitySuperClass}$$ implements Surfac
 		};
 		Log.debug("app restart : " + RestartExtra);
 //$${gameActivityOnRestartApplicationAdditions}$$
-		activityContext.runOnUiThread(restartRunnableUIThread);
+		runOnUiThread(restartRunnableUIThread);
 
 		// wait for exit.
 		while(true)
diff --git a/Engine/Build/Android/Java/src/com/epicgames/unreal/MessageBox01.java b/Engine/Build/Android/Java/src/com/epicgames/unreal/MessageBox01.java
--- a/Engine/Build/Android/Java/src/com/epicgames/unreal/MessageBox01.java
+++ b/Engine/Build/Android/Java/src/com/epicgames/unreal/MessageBox01.java
@@ -1,5 +1,4 @@
 // Copyright Epic Games, Inc. All Rights Reserved.
-
 package com.epicgames.unreal;
 
 import android.app.Dialog;
@@ -73,7 +72,25 @@ class MessageBox01
 	
 	protected void createAlert()
 	{
+		if (GameApplication.getAppContext() == null)
+		{
+			GameActivity.Log.error("GameApplication.getAppContext() is NULL and cannot create Alert for Caption=" + Caption + ", text=" + Text);
+			return;
+		}
+		if (!GameActivity.isValidGameActivity() || GameActivity.Get().getApplicationContext() == null || GameActivity.GetCurrentActivityContext() == null)
+		{
+			GameActivity.Log.error("GameActivity Application Context is not VALID yet! Cannot create Alert for Caption=" + Caption + ", text=" + Text);
+			return;
+		}
+
 		dismissAlert();
+		
+		if (GameActivity.isStandaloneMode())
+		{
+			GameActivity.Log.warn("AlertBox messages are not allowed in Standalone Mode (ASIS)! Caption=" + Caption + ", text=" + Text);
+			return;
+		}
+		
 		ButtonIdToIndex.clear();
 		Builder = new AlertDialog.Builder(GameActivity.Get());
 		Builder.setTitle(Caption);
diff --git a/Engine/Build/Android/Java/src/com/epicgames/unreal/SimpleContextWrapper.java b/Engine/Build/Android/Java/src/com/epicgames/unreal/SimpleContextWrapper.java
deleted file mode 100644
diff --git a/Engine/Source/Runtime/ApplicationCore/Private/Android/AndroidWindow.cpp b/Engine/Source/Runtime/ApplicationCore/Private/Android/AndroidWindow.cpp
--- a/Engine/Source/Runtime/ApplicationCore/Private/Android/AndroidWindow.cpp
+++ b/Engine/Source/Runtime/ApplicationCore/Private/Android/AndroidWindow.cpp
@@ -200,41 +200,26 @@ extern void GAndroidWindowLock_Unlock(FString calledBy);
 
 JNI_METHOD void Java_com_epicgames_makeaar_GameActivityForMakeAAR_nativeSetSurfaceOverride(JNIEnv* jenv, jobject thiz, jobject surface, jint x, jint y)
 {
+	ANativeWindow* prev = (ANativeWindow*)GAndroidWindowOverride;
 	if (surface != 0)
 	{
 		GSurfaceViewX = x;
 		GSurfaceViewY = y;
-		void* prev = GAndroidWindowOverride;
-		GAndroidWindowOverride = (ANativeWindow*)ANativeWindow_fromSurface(jenv, surface);
-		STANDALONE_DEBUG_LOG(TEXT("nativeSetSurfaceOverride(makeaar) applied: prev to new %p -> %p, (%d, %d)"), prev, GAndroidWindowOverride, GSurfaceViewX, GSurfaceViewY);
-		if (prev != GAndroidWindowOverride)
-		{
-			//GAndroidWindowLock_Unlock("Java_com_epicgames_makeaar_GameActivityForMakeAAR_nativeSetSurfaceOverride");
-		}
-
-		//if (GAndroidWindowOverride != nullptr)
-		//{
-		//	STANDALONE_DEBUG_LOGf(LogAndroid, TEXT("nativeSetSurfaceOverride. FTaskGraphInterface::IsRunning() = %d, GAndroidWindowOverride=%x"), FTaskGraphInterface::IsRunning(), GAndroidWindowOverride);
 
-		//	if (FTaskGraphInterface::IsRunning())
-		//	{
-		//		FSimpleDelegateGraphTask::CreateAndDispatchWhenReady(
-		//			FSimpleDelegateGraphTask::FDelegate::CreateLambda([=]()
-		//				{
-		//					STANDALONE_DEBUG_LOGf(LogAndroid, TEXT("from nativeSetSurfaceOverride, trigger UnlockAndroidWindow. FTaskGraphInterface::IsRunning() = %d, GAndroidWindowOverride=%x"), FTaskGraphInterface::IsRunning(), GAndroidWindowOverride);
-
-		//					FPlatformMisc::UnlockAndroidWindow();
-		//				}), TStatId(), nullptr, ENamedThreads::GameThread);
-		//	}
-		//}
+		GAndroidWindowOverride = (ANativeWindow*)ANativeWindow_fromSurface(jenv, surface);
+		UE_LOG(LogAndroid, Log, TEXT("nativeSetSurfaceOverride applied: prev to new %p -> %p, pos(%d, %d)"), prev, GAndroidWindowOverride, GSurfaceViewX, GSurfaceViewY);
 
 	}
 	else
 	{
 		GAndroidWindowOverride = nullptr;
 		
-		STANDALONE_DEBUG_LOG(TEXT("nativeSetSurfaceOverride(makeaar) setting to null and lock window"));
-		//GAndroidWindowLock_Lock("Java_com_epicgames_makeaar_GameActivityForMakeAAR_nativeSetSurfaceOverride");
+		STANDALONE_DEBUG_LOG(TEXT("nativeSetSurfaceOverride(makeaar) setting to null"));
+	}
+
+	if (prev != nullptr)
+	{
+		ANativeWindow_release(prev);
 	}
 }
 
@@ -566,7 +551,7 @@ FPlatformRect FAndroidWindow::GetScreenRect(bool bUseEventThreadWindow)
 	// too much of the following code needs JNI things, just assume override
 #if !USE_ANDROID_JNI
 
-	UE_LOG(LogAndroid, Fatal, TEXT("FAndroidWindow::CalculateSurfaceSize currently expedcts non-JNI platforms to override resolution"));
+	UE_LOG(LogAndroid, Fatal, TEXT("FAndroidWindow::CalculateSurfaceSize currently expects non-JNI platforms to override resolution"));
 	return FPlatformRect();
 #else
 
@@ -708,7 +693,13 @@ void FAndroidWindow::CalculateSurfaceSize(int32_t& SurfaceWidth, int32_t& Surfac
 	// do not convert to a surface size that is larger than native resolution
 	// Mobile VR doesn't need buffer quantization as Unreal never renders directly to the buffer in VR mode. 
 	static const bool bIsMobileVRApp = AndroidThunkCpp_IsOculusMobileApplication();
+	
+#if USE_ANDROID_STANDALONE
+	const int DividableBy = 1;	// don't change size of external window 
+#else
 	const int DividableBy = bIsMobileVRApp ? 1 : 8;
+#endif
+
 	SurfaceWidth = (SurfaceWidth / DividableBy) * DividableBy;
 	SurfaceHeight = (SurfaceHeight / DividableBy) * DividableBy;
 #endif
diff --git a/Engine/Source/Runtime/ApplicationCore/Public/Android/AndroidWindowUtils.h b/Engine/Source/Runtime/ApplicationCore/Public/Android/AndroidWindowUtils.h
--- a/Engine/Source/Runtime/ApplicationCore/Public/Android/AndroidWindowUtils.h
+++ b/Engine/Source/Runtime/ApplicationCore/Public/Android/AndroidWindowUtils.h
@@ -68,6 +68,7 @@ namespace AndroidWindowUtils
 		{
 			int32 Width = InOutScreenWidth;
 			int32 Height = InOutScreenHeight;
+			bool swapXY = false;
 
 			if (RequestedResX > 0)
 			{
@@ -76,6 +77,9 @@ namespace AndroidWindowUtils
 				{
 					Height = RequestedResX;
 					Width = FMath::TruncToInt32((float)Height * ((float)InOutScreenWidth / (float)InOutScreenHeight) + 0.5f);
+#if PLATFORM_ANDROID && USE_ANDROID_STANDALONE
+					swapXY = true;
+#endif
 				}
 				else
 				{
@@ -90,6 +94,9 @@ namespace AndroidWindowUtils
 				{
 					Width = RequestedResY;
 					Height = FMath::TruncToInt32((float)Width * ((float)InOutScreenHeight / (float)InOutScreenWidth) + 0.5f);
+#if PLATFORM_ANDROID && USE_ANDROID_STANDALONE
+					swapXY = true;
+#endif
 				}
 				else
 				{
@@ -105,6 +112,9 @@ namespace AndroidWindowUtils
 				if (InOutScreenHeight > InOutScreenWidth)
 				{
 					Height = FMath::TruncToInt32(1280.f * RequestedContentScaleFactor);
+#if PLATFORM_ANDROID && USE_ANDROID_STANDALONE
+					swapXY = true;
+#endif
 				}
 				else
 				{
@@ -115,7 +125,7 @@ namespace AndroidWindowUtils
 				Width = FMath::TruncToInt32((float)Height * AspectRatio + 0.5f);
 			}
 
-			const FIntVector2 Sanitized = SanitizeAndroidScreenSize(FIntVector2(InOutScreenWidth, InOutScreenHeight), FIntVector2(Width, Height));
+			const FIntVector2 Sanitized = SanitizeAndroidScreenSize(FIntVector2(InOutScreenWidth, InOutScreenHeight), !swapXY ? FIntVector2(Width, Height) : FIntVector2(Height, Width));
 			InOutScreenWidth = Sanitized.X;
 			InOutScreenHeight = Sanitized.Y;
 
diff --git a/Engine/Source/Runtime/Launch/Private/Android/AndroidEventManager.cpp b/Engine/Source/Runtime/Launch/Private/Android/AndroidEventManager.cpp
--- a/Engine/Source/Runtime/Launch/Private/Android/AndroidEventManager.cpp
+++ b/Engine/Source/Runtime/Launch/Private/Android/AndroidEventManager.cpp
@@ -97,7 +97,7 @@ void FAppEventManager::Tick()
 			break;
 		case APP_EVENT_STATE_WINDOW_DESTROYED:
 			bHaveWindow = false;
-			STANDALONE_DEBUG_LOGf(LogAndroidEvents, TEXT("APP_EVENT_STATE_WINDOW_DESTROYED, %d, %d, %d"), int(bRunning), int(bHaveWindow), int(bHaveGame));
+			STANDALONE_DEBUG_LOGf(LogAndroidEvents, TEXT("APP_EVENT_STATE_WINDOW_DESTROYED, %d, %d, %d"), int(bRunning), int(bWindowInFocus), int(bHaveGame));
 			break;
 		case APP_EVENT_STATE_ON_START:
 			//doing nothing here
@@ -114,9 +114,11 @@ void FAppEventManager::Tick()
 		case APP_EVENT_STATE_ON_PAUSE:
 			FAndroidAppEntry::OnPauseEvent();
 			bHaveGame = false;
+			STANDALONE_DEBUG_LOGf(LogAndroidEvents, TEXT("APP_EVENT_STATE_ON_PAUSE, %d, %d, %d"), int(bRunning), int(bWindowInFocus), int(bHaveGame));
 			break;
 		case APP_EVENT_STATE_ON_RESUME:
 			bHaveGame = true;
+			STANDALONE_DEBUG_LOGf(LogAndroidEvents, TEXT("APP_EVENT_STATE_ON_RESUME, %d, %d, %d"), int(bRunning), int(bWindowInFocus), int(bHaveGame));
 			break;
 
 		// window focus events that follow their own hierarchy, and might or might not respect App main events hierarchy
@@ -213,12 +215,16 @@ FAppEventManager::FAppEventManager():
 	CVarScale->SetOnChangedCallback(FConsoleVariableDelegate::CreateStatic(&FAppEventManager::OnScaleFactorChanged));
 
 	IConsoleVariable* CVarResX = IConsoleManager::Get().FindConsoleVariable(TEXT("r.Mobile.DesiredResX"));
-	check(CVarResX);
-	CVarResX->SetOnChangedCallback(FConsoleVariableDelegate::CreateStatic(&FAppEventManager::OnScaleFactorChanged));
+	if (CVarResX != nullptr)
+	{
+		CVarResX->SetOnChangedCallback(FConsoleVariableDelegate::CreateStatic(&FAppEventManager::OnScaleFactorChanged));
+	}
 
 	IConsoleVariable* CVarResY = IConsoleManager::Get().FindConsoleVariable(TEXT("r.Mobile.DesiredResY"));
-	check(CVarResY);
-	CVarResY->SetOnChangedCallback(FConsoleVariableDelegate::CreateStatic(&FAppEventManager::OnScaleFactorChanged));
+	if (CVarResY != nullptr)
+	{
+		CVarResY->SetOnChangedCallback(FConsoleVariableDelegate::CreateStatic(&FAppEventManager::OnScaleFactorChanged));
+	}
 }
 
 void FAppEventManager::OnScaleFactorChanged(IConsoleVariable* CVar)
@@ -234,16 +240,11 @@ void FAppEventManager::OnScaleFactorChanged(IConsoleVariable* CVar)
     }
 }
 
-static bool IsAppPausedOrSuspended()
-{
-	return FAppEventManager::GetInstance()->IsGamePaused() && !FAppEventManager::GetInstance()->IsGameInFocus();
-}
-
 void FAppEventManager::HandleWindowCreated_EventThread(void* InWindow)
 {
 	if (InWindow == nullptr)
 	{
-		STANDALONE_DEBUG_LOGf(LogAndroidEvents, TEXT("FAppEventManager::HandleWindowCreated_EventThread was given a NULL window so IGNORE."));
+		FPlatformMisc::LowLevelOutputDebugStringf(TEXT("FAppEventManager::HandleWindowCreated_EventThread was given a NULL window so IGNORE."));
 		return;
 	}
 
@@ -252,7 +253,6 @@ void FAppEventManager::HandleWindowCreated_EventThread(void* InWindow)
 	// Make sure window will not be deleted until event is processed
 	// Window could be deleted by OS while event queue stuck at game start-up phase
 	FAndroidWindow::AcquireWindowRef((ANativeWindow*)InWindow);
-
 	FAndroidWindow::SetHardwareWindow_EventThread(InWindow);
 
 	if (!AlreadyInited)
@@ -272,8 +272,8 @@ void FAppEventManager::HandleWindowClosed_EventThread()
 		return;
 	}
 
-	FAndroidWindow::SetHardwareWindow_EventThread(nullptr);
 	FAndroidWindow::ReleaseWindowRef((ANativeWindow*)ActiveWindow);
+	FAndroidWindow::SetHardwareWindow_EventThread(nullptr);
 
 	EnqueueAppEvent(APP_EVENT_STATE_WINDOW_DESTROYED);
 }
@@ -401,7 +401,11 @@ bool FAppEventManager::IsGamePaused()
 
 bool FAppEventManager::IsGameInFocus()
 {
+#if USE_ANDROID_STANDALONE
+	return (bWindowInFocus && bHaveWindow && bHaveGame);
+#else
 	return (bWindowInFocus && bHaveWindow);
+#endif
 }
 
 
diff --git a/Engine/Source/Runtime/Launch/Private/Android/LaunchAndroid.cpp b/Engine/Source/Runtime/Launch/Private/Android/LaunchAndroid.cpp
--- a/Engine/Source/Runtime/Launch/Private/Android/LaunchAndroid.cpp
+++ b/Engine/Source/Runtime/Launch/Private/Android/LaunchAndroid.cpp
@@ -291,86 +291,11 @@ static bool bShouldRestartFromInterrupt = false;
 static bool bIgnorePauseOnDownloaderStart = false;
 bool bReadyToProcessEvents = false;
 
-#if USE_ANDROID_STANDALONE
-static int GAndroidWindowLockRefCount = 0;
-static const bool bForce_GAndroidWindowLock = true;
-#endif
-
 pthread_t G_AndroidEventThread;
 struct android_app* GNativeAndroidApp = NULL;
+static uint32 EventThreadID = 0;
 
-void GAndroidWindowLock_Lock(FString calledBy)
-{
-#if !USE_ANDROID_STANDALONE
-	GAndroidWindowLock.Lock();
-#else
-	bool doAction = bForce_GAndroidWindowLock;
-	if (GAndroidWindowLockRefCount >= 0)
-	{
-		++GAndroidWindowLockRefCount;
-		if (GAndroidWindowLockRefCount == 1)
-		{
-			doAction = true;
-		}
-		else
-		{
-			STANDALONE_DEBUG_LOGf(LogAndroid, TEXT("WARNING GAndroidWindowLock_Lock(%s). ALREADY LOCKED. GAndroidWindowLockRefCount=%d, bAppIsActive_EventThread=%d"), *calledBy, GAndroidWindowLockRefCount, bAppIsActive_EventThread);
-		}
-	}
-	else
-	{
-		STANDALONE_DEBUG_LOGf(LogAndroid, TEXT("GAndroidWindowLock_Lock(%s). ERROR: GAndroidWindowLockRefCount=%d, bAppIsActive_EventThread=%d"), *calledBy, GAndroidWindowLockRefCount, bAppIsActive_EventThread);
-	}
-	if (doAction)
-	{
-		GAndroidWindowLock.Lock();
-
-		STANDALONE_DEBUG_LOGf(LogAndroid, TEXT("GAndroidWindowLock_Lock(%s). LOCKED. GAndroidWindowLockRefCount=%d, bAppIsActive_EventThread=%d"), *calledBy, GAndroidWindowLockRefCount, bAppIsActive_EventThread);
-	}
-	else
-	{
-		STANDALONE_DEBUG_LOGf(LogAndroid, TEXT("WARNING GAndroidWindowLock_Lock(%s). NOOP: GAndroidWindowLockRefCount=%d, bAppIsActive_EventThread=%d bReadyToProcessEvents=%d"), *calledBy, GAndroidWindowLockRefCount, bAppIsActive_EventThread, bReadyToProcessEvents);
-	}
-#endif
-}
 
-void GAndroidWindowLock_Unlock(FString calledBy)
-{
-#if !USE_ANDROID_STANDALONE
-	GAndroidWindowLock.Unlock();
-#else
-	bool doAction = bForce_GAndroidWindowLock;
-	if (GAndroidWindowLockRefCount > 0)
-	{
-		if (--GAndroidWindowLockRefCount == 0)
-		{
-			doAction = true;
-		}
-		else
-		{
-			STANDALONE_DEBUG_LOGf(LogAndroid, TEXT("WARNING GAndroidWindowLock_Unlock(%s). KEEP LOCKED, GAndroidWindowLockRefCount=%d, bAppIsActive_EventThread=%d"), *calledBy, GAndroidWindowLockRefCount, bAppIsActive_EventThread);
-		}
-		if (GAndroidWindowLockRefCount < 0)
-		{
-			STANDALONE_DEBUG_LOGf(LogAndroid, TEXT("GAndroidWindowLock_Unlock(%s). ERROR UNLOCK MISMATCH, GAndroidWindowLockRefCount=%d, bAppIsActive_EventThread=%d"), *calledBy, GAndroidWindowLockRefCount, bAppIsActive_EventThread);
-			GAndroidWindowLockRefCount = 0;
-		}
-	}
-	else
-	{
-		STANDALONE_DEBUG_LOGf(LogAndroid, TEXT("GAndroidWindowLock_Unlock(%s). SKIPPED, GAndroidWindowLockRefCount=%d, bAppIsActive_EventThread=%d"), *calledBy, GAndroidWindowLockRefCount, bAppIsActive_EventThread);
-	}
-	if (doAction)
-	{
-		GAndroidWindowLock.Unlock();
-		STANDALONE_DEBUG_LOGf(LogAndroid, TEXT("GAndroidWindowLock_Unlock(%s). UNLOCK, GAndroidWindowLockRefCount=%d, bAppIsActive_EventThread=%d"), *calledBy, GAndroidWindowLockRefCount, bAppIsActive_EventThread);
-	}
-	else
-	{
-		STANDALONE_DEBUG_LOGf(LogAndroid, TEXT("GAndroidWindowLock_Unlock(%s). NOOP: GAndroidWindowLockRefCount=%d, bAppIsActive_EventThread=%d bReadyToProcessEvents=%d"), *calledBy, GAndroidWindowLockRefCount, bAppIsActive_EventThread, bReadyToProcessEvents);
-	}
-#endif
-}
 static void SuspendApp_EventThread();
 static void ActivateApp_EventThread();
 #if USE_ANDROID_STANDALONE && !(UE_BUILD_SHIPPING || UE_BUILD_TEST)
@@ -387,7 +312,7 @@ void FPlatformMisc::UnlockAndroidWindow()
 	check(FTaskGraphInterface::IsRunning());
 	
 	bReadyToProcessEvents = true;
-	GAndroidWindowLock_Unlock("UnlockAndroidWindow");
+	GAndroidWindowLock.Unlock();
 }
 
 JNI_METHOD void Java_com_epicgames_unreal_GameActivity_nativeResumeMainInit(JNIEnv* jenv, jobject thiz)
@@ -484,6 +409,7 @@ static void InitCommandLine()
 			FCommandLine::Append(UTF8_TO_TCHAR(CommandLine));
 			FPlatformMisc::LowLevelOutputDebugStringf(TEXT("APK Commandline: %s"), FCommandLine::Get());
 		}
+
 #if !UE_BUILD_SHIPPING && USE_ANDROID_STANDALONE
 		UE_SET_LOG_VERBOSITY(LogAndroid, Log);
 		if (!GAndroidCommandLine.IsEmpty())
@@ -596,21 +522,31 @@ static void ApplyAndroidCompatConfigRules()
 }
 
 //Main function called from the android entry point
-#if USE_ANDROID_STANDALONE
+int32 AndroidMain(struct android_app* state);
+
 void* AndroidMain(void* param)
 {
 	struct android_app* state = (struct android_app*)param;
-
 	FTaskTagScope Scope(ETaskTag::EGameThread);
-	GGameThreadId = FPlatformTLS::GetCurrentThreadId();
-	GNativeAndroidApp = state;
-	check(GNativeAndroidApp);
 
-#else
+	if (GGameThreadId == 0)
+	{
+		GGameThreadId = FPlatformTLS::GetCurrentThreadId();
+		STANDALONE_DEBUG_LOGf(LogAndroid, TEXT("AndroidMain set current GGameThreadId=%d"), GGameThreadId);
+	}
+
+	if (EventThreadID == 0)
+	{
+		EventThreadID = GGameThreadId;
+		STANDALONE_DEBUG_LOGf(LogAndroid, TEXT("AndroidMain set current EventThreadID=%d"), EventThreadID);
+	}
+	AndroidMain(state);
+	return nullptr;
+}
+
+
 int32 AndroidMain(struct android_app* state)
 {
-#endif
-
 	BootTimingPoint("AndroidMain");
 
 	FPlatformMisc::LowLevelOutputDebugString(TEXT("Entered AndroidMain()\n"));
@@ -700,14 +636,7 @@ int32 AndroidMain(struct android_app* state)
 			FPlatformMisc::MemoryBarrier();
 		}
 	}
-	// make sure GEngineLoop::Exit() is always called.
-	struct EngineLoopCleanup
-	{
-		~EngineLoopCleanup()
-		{
-			EngineExit();
-		}
-	} Cleanup;
+
 
 	// read the command line file
 	InitCommandLine();
@@ -745,7 +674,7 @@ int32 AndroidMain(struct android_app* state)
 		// We don't wait for all of preinit to complete as PreLoadScreens will need to process events during preinit.
 
 		UE_LOG(LogAndroid, Log, TEXT("PreInit android HW window lock. bAppIsActive_EventThread=%d"), bAppIsActive_EventThread);
-		GAndroidWindowLock_Lock("AndroidMain lock");
+		GAndroidWindowLock.Lock();
 	}
 
 	FPlatformMisc::LowLevelOutputDebugString(TEXT("After GAndroidWindowLock in AndroidMain"));
@@ -760,11 +689,7 @@ int32 AndroidMain(struct android_app* state)
 	if (PreInitResult != 0)
 	{
 		checkf(false, TEXT("Engine Preinit Failed"));
-#if USE_ANDROID_STANDALONE
-		return nullptr;
-#else
 		return PreInitResult;
-#endif
 	}
 
 	// register callback for native window resize
@@ -821,7 +746,11 @@ int32 AndroidMain(struct android_app* state)
 		FAndroidStats::UpdateAndroidStats();
 
 		FAppEventManager::GetInstance()->Tick();
-		if (!FAppEventManager::GetInstance()->IsGamePaused())
+		if (!FAppEventManager::GetInstance()->IsGamePaused()
+#if USE_ANDROID_STANDALONE
+			//&& FAppEventManager::GetInstance()->IsGameInFocus()
+#endif
+			)
 		{
 			GEngineLoop.Tick();
 		}
@@ -848,18 +777,12 @@ int32 AndroidMain(struct android_app* state)
 
 	// exit out!
 	GEngineLoop.Exit();
-	
-#if !USE_ANDROID_STANDALONE
-	FPlatformMisc::RequestExit(true, TEXT("AndroidMain"));
-#endif
 
 	UE_LOG(LogAndroid, Log, TEXT("Exiting is over"));
 
-#if USE_ANDROID_STANDALONE
-	return nullptr;
-#else
-	return ErrorLevel;
-#endif
+	FPlatformMisc::RequestExit(true, TEXT("AndroidMain"));
+
+	return 0;
 }
 
 struct AChoreographer;
@@ -973,10 +896,11 @@ void FChoreographer::SetCallback(int64 Delay)
 	AChoreographer_postFrameCallbackDelayed_(choreographer, choreographer_callback, nullptr, Delay / 1000000);
 }
 
-static uint32 EventThreadID = 0;
-
 bool IsInAndroidEventThread()
 {
+	// Note: leave the commented out line for debug purposes.
+	//STANDALONE_DEBUG_LOGf(LogAndroid, TEXT("IsInAndroidEventThread(), GGameThreadId=%d, EventThreadID=%d, FPlatformTLS::GetCurrentThreadId()=%d"), GGameThreadId, EventThreadID, FPlatformTLS::GetCurrentThreadId());
+
 #if USE_ANDROID_STANDALONE
 	//@TODO: for now always return true to avoid check failures
 	return true;
@@ -989,6 +913,8 @@ bool IsInAndroidEventThread()
 static void* AndroidEventThreadWorker( void* param )
 {
 	pthread_setname_np(pthread_self(), "EventWorker");
+	STANDALONE_DEBUG_LOGf(LogAndroid, TEXT("AndroidEventThreadWorker(begin), GGameThreadId=%d, EventThreadID=%d, FPlatformTLS::GetCurrentThreadId()=%d"), GGameThreadId, EventThreadID, FPlatformTLS::GetCurrentThreadId());
+
 	EventThreadID = FPlatformTLS::GetCurrentThreadId();
 	FAndroidMisc::RegisterThreadName("EventWorker", EventThreadID);
 
@@ -1015,11 +941,9 @@ static void* AndroidEventThreadWorker( void* param )
 
 	TheChoreographer.SetupChoreographer();
 
-	GEventHandlerInitialized = true;
-
 	// window is initially invalid/locked.
 	UE_LOG(LogAndroid, Log, TEXT("AndroidEventThreadWorker, Initial HW window lock."));
-	GAndroidWindowLock_Lock("AndroidEventThreadWorker");
+	GAndroidWindowLock.Lock();
 
 	DEVELOPER_LOG_COMMANDCB_CASE(AndroidEventThreadWorker_BeforeWhile);
 
@@ -1035,7 +959,7 @@ static void* AndroidEventThreadWorker( void* param )
 	}
 	DEVELOPER_LOG_COMMANDCB_CASE(AndroidEventThreadWorker_AfterWhile);
 
-	GAndroidWindowLock_Unlock("AndroidEventThreadWorker");
+	GAndroidWindowLock.Unlock();
 
 	UE_LOG(LogAndroid, Log, TEXT("AndroidEventThreadWorker->Exiting"));
 	STANDALONE_DEBUG_LOGf(LogAndroid, TEXT("Exiting AndroidEventThreadWorker"));
@@ -1091,6 +1015,7 @@ void android_main(struct android_app* state)
 	//@todo android: replace with native activity, main loop off of UI thread, etc.
 	AndroidMain(state);
 }
+
 void* android_main(void* param)
 {
 	struct android_app* state = (struct android_app*)param;
@@ -1525,22 +1450,23 @@ FAppEventData::FAppEventData(ANativeWindow* WindowIn)
 	//check(WindowWidth >= 0 && WindowHeight >= 0);
 }
 
-
-
 // called when the app has window + resume.
 static void ActivateApp_EventThread()
 {
 	DEVELOPER_LOG_COMMANDCB_CASE(ActivateApp_EventThread);
 	if (bAppIsActive_EventThread)
 	{
-		// Seems this can occur.
 		STANDALONE_DEBUG_LOGf(LogAndroid, TEXT("UNEXPECTED -- event thread, activate app, ALREADY have HW window lock according to bAppIsActive_EventThread. bReadyToProcessEvents=%d, bAppIsActive_EventThread=%d"), bReadyToProcessEvents, bAppIsActive_EventThread);
+
+#if !USE_ANDROID_STANDALONE // SUSPECT TODO should remove the #if?
+		// Seems this can occur.
 		return;
+#endif
 	}
 
 	// Unlock window when we're ready.
 	UE_LOG(LogAndroid, Log, TEXT("event thread, activate app, unlocking HW window"));
-	GAndroidWindowLock_Unlock("ActivateApp_EventThread");
+	GAndroidWindowLock.Unlock();
 
 	// wake the GT up.
 	FAppEventManager::GetInstance()->EnqueueAppEvent(APP_EVENT_STATE_APP_ACTIVATED);
@@ -1570,7 +1496,7 @@ static void ActivateApp_EventThread()
 }
 
 extern void BlockRendering();
-// called whenever the app loses loses window or pause.
+// called whenever the app loses window or pause.
 static void SuspendApp_EventThread()
 {
 	DEVELOPER_LOG_COMMANDCB_CASE(SuspendApp_EventThread);
@@ -1584,7 +1510,7 @@ static void SuspendApp_EventThread()
 	// Lock the window, this prevents event thread from removing the window whilst the RHI initializes.
 
 	UE_LOG(LogAndroid, Log, TEXT("SuspendApp_EventThread -> event thread, suspending app, acquiring HW window lock."));
-	GAndroidWindowLock_Lock("SuspendApp_EventThread");
+	GAndroidWindowLock.Lock();
 	STANDALONE_DEBUG_LOGf(LogAndroid, TEXT("SuspendApp_EventThread -> event thread, suspending app, acquiring HW window lock. bReadyToProcessEvents=%d, bAppIsActive_EventThread=%d"), bReadyToProcessEvents, bAppIsActive_EventThread);
 
 	if (bReadyToProcessEvents == false)
@@ -1603,7 +1529,6 @@ static void SuspendApp_EventThread()
 		FPlatformProcess::ReturnSynchEventToPool(EventToDelete);
 	});
 	
-
 	// perform the delegates before the window handle is cleared.
 	// This ensures any tasks that require a window handle will have it before we block the RT on the invalid window.
 	FAppEventManager::GetInstance()->EnqueueAppEvent(APP_EVENT_RUN_CALLBACK, FAppEventData([EMDoneTrigger]()
@@ -1629,9 +1554,12 @@ static void SuspendApp_EventThread()
 	// wait for a period of time before blocking rendering
 	UE_LOG(LogAndroid, Log, TEXT("SuspendApp_EventThread -> , waiting for event manager to process. tid: %d"), FPlatformTLS::GetCurrentThreadId());
 #if USE_ANDROID_STANDALONE
-	EMDoneTrigger->Reset();
-#endif
+	//EMDoneTrigger->Reset();
+	bool bSuccess = EMDoneTrigger->Wait(240);
+#else
 	bool bSuccess = EMDoneTrigger->Wait(4000);
+#endif
+
 	float ElapsedTimeInMs_EMDoneTrigger_Wait = FPlatformTime::ToMilliseconds(FPlatformTime::Cycles() - StartCycles);
 	UE_CLOG(!bSuccess, LogAndroid, Log, TEXT("SuspendApp_EventThread -> ERROR: backgrounding callback, not responded in timely manner. EMDoneTrigger->Wait, waited '%f' ms"), (float)ElapsedTimeInMs_EMDoneTrigger_Wait);
 	UE_LOG(LogAndroid, Log, TEXT("SuspendApp_EventThread -> EMDoneTrigger->Wait, waited '%f' ms"), (float)ElapsedTimeInMs_EMDoneTrigger_Wait);
@@ -1693,12 +1621,13 @@ static void OnAppCommandCB(struct android_app* app, int32_t cmd)
 			}
 			checkf(win != NULL, TEXT("Engine APP_CMD_INIT_WINDOW with STANDALONE Failed since win is NULL!"));
 			FAppEventManager::GetInstance()->HandleWindowCreated_EventThread((ANativeWindow*)win);
+			bHasWindow = win != nullptr;
 		}
 #else
 		STANDALONE_DEBUG_LOGf(LogAndroid, TEXT("Case APP_CMD_INIT_WINDOW"));
 		FAppEventManager::GetInstance()->HandleWindowCreated_EventThread(app->pendingWindow);
-#endif
 		bHasWindow = true;
+#endif
 		UE_LOG(LogAndroid, Log, TEXT("Case APP_CMD_INIT_WINDOW: bHasWindow=%d, bHasFocus=%d, bIsResumed=%d"), bHasWindow, bHasFocus, bIsResumed);
 		if (bHasWindow && bHasFocus && bIsResumed)
 		{
@@ -1814,8 +1743,6 @@ static void OnAppCommandCB(struct android_app* app, int32_t cmd)
 		break;
 	case APP_CMD_RESUME:
 
-		DEVELOPER_LOG_COMMANDCB_CASE(APP_CMD_RESUME);
-
 		/**
 		 * Command from main thread: the app's activity has been resumed.
 		 */
@@ -2098,7 +2025,6 @@ bool WaitForAndroidLoseFocusEvent(double TimeoutSeconds)
 	return FAppEventManager::GetInstance()->WaitForEventInQueue(EAppEventState::APP_EVENT_STATE_WINDOW_LOST_FOCUS, TimeoutSeconds);
 }
 
-
 #if USE_ANDROID_STANDALONE
 pthread_t G_AndroidMainThread;
 
@@ -2175,7 +2101,7 @@ static void InitEvent()
 
 static void AppSuspended()
 {
-	STANDALONE_DEBUG_LOGf(LogAndroid, TEXT("LaunchAndroid received AppSuspended"));
+	STANDALONE_DEBUG_LOGf(LogAndroid, TEXT("LaunchAndroid received AppSuspended, GFrameCounter=%d"), (int32)GFrameCounter);
 	AndroidThunkCpp_Engine_SendEvent(EVENTTYPE_APP_SUSPENDED, TEXT(""), (int32)GFrameCounter, 0, 0.0f);
 }
 
@@ -2247,6 +2173,7 @@ JNI_METHOD void Java_com_epicgames_makeaar_GameActivityForMakeAAR_nativeMain(JNI
 	FCoreDelegates::OnInit.AddStatic(InitEvent);
 	FCoreDelegates::OnPostEngineInit.AddStatic(PostEngineInitEvent);
 	FCoreDelegates::OnFEngineLoopInitComplete.AddStatic(EngineLoopInitCompleteEvent);
+	FCoreDelegates::ApplicationWillEnterBackgroundDelegate.AddStatic(AppSuspended);
 	FCoreDelegates::OnBeginFrame.AddStatic(BeginFrameEvent);
 	FCoreDelegates::OnEndFrame.AddStatic(EndFrameEvent);
 	FCoreUObjectDelegates::PreLoadMap.AddStatic(PreLoadMapEvent);
@@ -2271,9 +2198,9 @@ JNI_METHOD jint Java_com_epicgames_makeaar_GameActivityForMakeAAR_nativeInputTou
 	TouchType type = TouchEnded;
 	switch (inType)
 	{
-	case 0: type = TouchBegan; break;
-	case 1: type = TouchMoved; break;
-	case 2: type = TouchEnded; break;
+		case 0: type = TouchBegan; break;
+		case 1: type = TouchMoved; break;
+		case 2: type = TouchEnded; break;
 	}
 
 	ANativeWindow* Window = (ANativeWindow*)FAndroidWindow::GetHardwareWindow_EventThread();
@@ -2379,7 +2306,7 @@ JNI_METHOD bool Java_com_epicgames_makeaar_Engine_nativeMountPak(JNIEnv* jenv, j
 
 JNI_METHOD void Java_com_epicgames_makeaar_GameActivityForMakeAAR_nativeAppCommand(JNIEnv* jenv, jobject thiz, jint cmd)
 {
-	STANDALONE_DEBUG_LOGf(LogAndroid, TEXT("MAKEAAR -> nativeAppCommand[%u, %s] -> EventThreadID=%d, GGameThreadId=%d, GResumeMainInit=%d, GEventHandlerInitialized=%d, FTaskGraphInterface::IsRunning()=%d, FPlatformTLS::GetCurrentThreadId()=%d"), 
+	STANDALONE_DEBUG_LOGf(LogAndroid, TEXT("GameActivityForMakeAAR -> nativeAppCommand[%u, %s] -> EventThreadID=%d, GGameThreadId=%d, GResumeMainInit=%d, GEventHandlerInitialized=%d, FTaskGraphInterface::IsRunning()=%d, FPlatformTLS::GetCurrentThreadId()=%d"), 
 		cmd,
 		get_APP_CMD_String(cmd),
 		EventThreadID, 
@@ -2407,4 +2334,4 @@ extern "C" {
     }
 };
 
-#endif
+#endif //  USE_ANDROID_LAUNCH
\ No newline at end of file
diff --git a/Engine/Source/Runtime/OpenGLDrv/Private/Android/AndroidEGL.cpp b/Engine/Source/Runtime/OpenGLDrv/Private/Android/AndroidEGL.cpp
--- a/Engine/Source/Runtime/OpenGLDrv/Private/Android/AndroidEGL.cpp
+++ b/Engine/Source/Runtime/OpenGLDrv/Private/Android/AndroidEGL.cpp
@@ -1260,9 +1260,7 @@ bool AndroidEGL::IsOfflineSurfaceRequired()
 }
 
 ///
-//extern FCriticalSection GAndroidWindowLock;
-extern void GAndroidWindowLock_Lock(FString calledBy);
-extern void GAndroidWindowLock_Unlock(FString calledBy);
+extern FCriticalSection GAndroidWindowLock;
 
 void BlockOnLostWindowRenderCommand(TSharedPtr<FEvent, ESPMode::ThreadSafe> RTBlockedTrigger)
 {
@@ -1292,7 +1290,7 @@ void BlockOnLostWindowRenderCommand(TSharedPtr<FEvent, ESPMode::ThreadSafe> RTBl
 
 		RTBlockedTrigger->Trigger();
 
-		GAndroidWindowLock_Lock("BlockOnLostWindowRenderCommand Vulkan");
+		GAndroidWindowLock.Lock();
 		UE_LOG(LogAndroid, Log, TEXT("RendererBlock acquired window lock"));
 		const auto& OnReinitWindowCallback = FAndroidMisc::GetOnReInitWindowCallback();
 		if (OnReinitWindowCallback)
@@ -1300,17 +1298,17 @@ void BlockOnLostWindowRenderCommand(TSharedPtr<FEvent, ESPMode::ThreadSafe> RTBl
 			OnReinitWindowCallback(FAndroidWindow::GetHardwareWindow_EventThread());
 			UE_LOG(LogAndroid, Log, TEXT("RendererBlock updating window"));
 		}
-		GAndroidWindowLock_Unlock("BlockOnLostWindowRenderCommand Vulkan");
+		GAndroidWindowLock.Unlock();
 	}
 	else
 	{
 		RunOnGLRenderContextThread([&] {
 			RTBlockedTrigger->Trigger();
-			GAndroidWindowLock_Lock("BlockOnLostWindowRenderCommand");
+			GAndroidWindowLock.Lock();
 			UE_LOG(LogAndroid, Log, TEXT("RendererBlock acquired window lock"));
 			AndroidEGL::GetInstance()->SetRenderContextWindowSurface();
 			UE_LOG(LogAndroid, Log, TEXT("RendererBlock updating window"));
-			GAndroidWindowLock_Unlock("BlockOnLostWindowRenderCommand");
+			GAndroidWindowLock.Unlock();
 		}, true);
 	}
 	UE_LOG(LogAndroid, Log, TEXT("RendererBlock released window lock"));
@@ -1328,9 +1326,7 @@ void SetSharedContextGameCommand(TSharedPtr<FEvent, ESPMode::ThreadSafe> GTBlock
 extern bool IsInAndroidEventThread();
 void BlockRendering()
 {
-#if !USE_ANDROID_STANDALONE
 	check(IsInAndroidEventThread());
-#endif
 	check(GIsRHIInitialized);
 
 	UE_LOG(LogAndroid, Log, TEXT("Blocking renderer on suspended window."));
