New: cleanup ASIS impact on GameActivity
QA: should check this change to ensure the 4 finger touch and other crashes are addressed.
#android

diff --git a/Engine/Build/Android/Java/src/com/epicgames/unreal/GameActivity.java.template b/Engine/Build/Android/Java/src/com/epicgames/unreal/GameActivity.java.template
--- a/Engine/Build/Android/Java/src/com/epicgames/unreal/GameActivity.java.template
+++ b/Engine/Build/Android/Java/src/com/epicgames/unreal/GameActivity.java.template
@@ -4,6 +4,7 @@
 package com.epicgames.unreal;
 
 import android.app.NotificationManager;
+import android.content.ContextWrapper;
 import android.opengl.EGL14;
 import android.opengl.EGLConfig;
 import android.opengl.EGLContext;
@@ -11,6 +12,7 @@ import android.opengl.EGLDisplay;
 import android.opengl.EGLSurface;
 import android.opengl.GLES20;
 
+import java.lang.annotation.Native;
 import java.util.regex.Pattern;
 import java.util.regex.Matcher;
 
@@ -339,8 +341,8 @@ public class GameActivity extends $${gameActivitySuperClass}$$ implements Surfac
 	public static final boolean bSamsungDevice = android.os.Build.MANUFACTURER.equals("samsung");
 	
 	protected boolean isStandalone = false;
-	protected static GameActivity _gameActivityInstance = null;
-	protected SimpleContextWrapper activityContext = null;
+	protected static GameActivity _activity = null;
+	protected ContextWrapper activityContext = null;
 	protected static Bundle _bundle;
 	protected static Bundle _extrasBundle;
 	protected boolean bOnCreateCalled = false;
@@ -539,14 +541,7 @@ public class GameActivity extends $${gameActivitySuperClass}$$ implements Surfac
 		Keyboard
 	}
 	private EAlertDialogType CurrentDialogType = EAlertDialogType.None;
-	
-	public GameActivity()
-	{
-		Log.debug( "GameActivity constructor called with _gameActivityInstance=" + _gameActivityInstance + ", activityContext=" + activityContext + ", this=" + this + ", isStandalone=" + isStandalone);
 
-		assert (_gameActivityInstance == null);
-		_gameActivityInstance = this;
-	}
 	
 	public boolean IsInVRMode()
 	{
@@ -556,32 +551,49 @@ public class GameActivity extends $${gameActivitySuperClass}$$ implements Surfac
 	/** Access singleton activity for game. **/
 	public static GameActivity Get()
 	{
-		return _gameActivityInstance;
+		return _activity;
 	}
 
-	public static Boolean isValidGameActivity()
+	public static Boolean isStandaloneMode()
 	{
-		return _gameActivityInstance != null;
+		return _activity.isStandalone;
 	}
 
-	public static SimpleContextWrapper GetCurrentActivityContext()
+	
+	public static Boolean isValidGameActivity()
 	{
-		return _gameActivityInstance.activityContext;
+		return _activity != null;
 	}
 
-	public final Window getWindow() 
+	public static ContextWrapper getCurrentActivityContext()
 	{
-		return isStandalone ? activityContext.getWindow() : super.getWindow();
+		assert(_activity != null && _activity.activityContext != null);
+		return _activity.activityContext;
 	}
 
-	public final Intent getIntent()
+	public final void runOnMainThread(Runnable action) {
+		assert (getBaseContext() != null);
+		assert (getMainExecutor() != null);
+		getMainExecutor().execute(action);
+	}
+	
+	private Activity getThisActivity()
 	{
-		return super.getIntent();
+		//Log.debug( "LIFECYCLE: GameActivity.getThisActivity(), bOnCreateCalled=" + bOnCreateCalled + ", _activity=" + _activity + ", activityContext=" + activityContext);
+
+		return this;
 	}
 
 	public final WindowManager getWindowManager()
 	{
-		return isStandalone ? activityContext.getWindowManager() : super.getWindowManager();
+		if (isStandalone) 
+		{
+			return (WindowManager) getThisActivity().getSystemService(Context.WINDOW_SERVICE);
+		}
+		else 
+		{
+			return super.getWindowManager();
+		}
 	}
 
 	
@@ -673,8 +685,7 @@ public class GameActivity extends $${gameActivitySuperClass}$$ implements Surfac
 
 	public String getProductName()
 	{
-	
-		String ProductName = Settings.Global.getString(getApplicationContext().getContentResolver(), "default_device_name");
+		String ProductName = Settings.Global.getString(getContentResolver(), "default_device_name");
 		if(ProductName == null)
 		{
 			ProductName="";
@@ -738,7 +749,7 @@ public class GameActivity extends $${gameActivitySuperClass}$$ implements Surfac
 	public String AndroidThunkJava_GetCommandLine()
 	{
 		String commandLine = "";
-		//$${gameActivityGetCommandLineAdditions}$$
+//$${gameActivityGetCommandLineAdditions}$$
 		return commandLine;
 	}
 
@@ -2429,16 +2440,42 @@ public class GameActivity extends $${gameActivitySuperClass}$$ implements Surfac
 		}
 		return foundStrings;
 	}
+	
+	private static Activity getActivityOfContext(Context context)
+	{
+		if (context == null)
+		{
+			return null;
+		}
+
+		while (context instanceof ContextWrapper) 
+		{
+			if (context instanceof Activity) 
+			{
+				return (Activity)context;
+			}
+			context = ((ContextWrapper)context).getBaseContext();
+		}
+
+		return null;
+	}
+
+
+
 
 	@Override
 	public void onCreate(Bundle savedInstanceState)
 	{
-		assert(_gameActivityInstance != null);
+		Log.debug("LIFECYCLE: GameActivity.onCreate() called!, this=" + this);
+				
+		_activity = this;
 		bOnCreateCalled = true;
-		activityContext = new SimpleContextWrapper(this);
+		activityContext = getActivityOfContext(this);
+		assert( activityContext != null);
+		
 		InternalFilesDir = getFilesDir().getAbsolutePath() + "/";
 		ExternalFilesDir = getExternalFilesDir(null).getAbsolutePath() + "/";
-		//$${gameActivityOnCreateBeginningAdditions}$$
+//$${gameActivityOnCreateBeginningAdditions}$$
 		Logger.RegisterCallback(this);
 
 		// Grab a reference to the asset manager
@@ -2451,6 +2488,40 @@ public class GameActivity extends $${gameActivitySuperClass}$$ implements Surfac
 
 	protected void onCreateBody(Bundle savedInstanceState)
 	{
+		if (isStandalone)
+		{
+			if (super.getBaseContext() == null)
+			{
+				Log.debug("LIFECYCLE: onCreateBody(): super.getBaseContext() is null so attach base context = " + activityContext);
+				attachBaseContext(activityContext);
+			}
+		
+			Log.debug("LIFECYCLE: onCreateBody(): isStandalone=" + isStandalone 
+				+ ", getBaseContext() = " + getBaseContext() 
+				+ ", activityContext = "+ activityContext);
+		
+			if (super.getApplicationContext() == null)
+			{
+				Log.debug("LIFECYCLE: onCreateBody(): super.getApplicationContext() is null so attach base context = " + activityContext);
+				attachBaseContext(activityContext);
+			}
+
+		
+			GameApplication gameApplication = GameApplication.getAppContext() instanceof GameApplication ? (GameApplication) GameApplication.getAppContext() : null;
+			if (gameApplication != null && gameApplication.getBaseContext() == null)
+			{
+				Log.debug("LIFECYCLE: onCreateBody(): attachBaseContext for GameApplication to context = " + activityContext);
+				gameApplication.attachBaseContext(activityContext);
+			}
+		}
+
+		if (isStandalone)
+		{
+			Log.debug("LIFECYCLE: onCreateBody(): getBaseContext() = " + getBaseContext());
+			Log.debug("LIFECYCLE: onCreateBody(): GameApplication.getAppContext() = " + GameApplication.getAppContext());
+			Log.debug("LIFECYCLE: onCreateBody(): activityContext.getApplicationContext() = " + activityContext.getApplicationContext());
+		}
+
 		sensorManager = (SensorManager) getSystemService(Context.SENSOR_SERVICE);
 		accelerometer = sensorManager.getDefaultSensor(Sensor.TYPE_ACCELEROMETER);
 		magnetometer = sensorManager.getDefaultSensor(Sensor.TYPE_MAGNETIC_FIELD);
@@ -2475,8 +2546,11 @@ public class GameActivity extends $${gameActivitySuperClass}$$ implements Surfac
 			MultiCastLock.setReferenceCounted(true);
 		}
 
+		Log.debug("LIFECYCLE: GameActivity.onCreateBody(), activityContext = " + activityContext);
+				
+
 		// disable autofill to avoid possible Android 10 crash (android.view.autofill.AutofillManager.ensureServiceClientAddedIfNeededLocked() SyncResultReceiver$TimeoutException)
-		if (Build.VERSION.SDK_INT >= Build.VERSION_CODES.O)
+		if (!isStandalone && Build.VERSION.SDK_INT >= Build.VERSION_CODES.O)
 		{
 			try
 			{
@@ -2550,7 +2624,7 @@ public class GameActivity extends $${gameActivitySuperClass}$$ implements Surfac
 				CommandLineFull = CmdLine;
 			}
 
-			if (UseDisplayCutout)
+			if (!isStandalone && UseDisplayCutout)
 			{
 				// will not be true if not Android Pie or later
 	            WindowManager.LayoutParams params = getWindow().getAttributes();
@@ -2606,7 +2680,7 @@ public class GameActivity extends $${gameActivitySuperClass}$$ implements Surfac
 			{
 				int ProcessMemory = 0;
 
-				ActivityManager activityManager = (ActivityManager) getSystemService(Context.ACTIVITY_SERVICE);
+				ActivityManager activityManager = (ActivityManager)_activity.getSystemService(Context.ACTIVITY_SERVICE);
 				int pid = android.os.Process.myPid();
 				int pids[] = new int[] { pid };
 				android.os.Debug.MemoryInfo[] memoryInfo = activityManager.getProcessMemoryInfo(pids);
@@ -2653,25 +2727,23 @@ public class GameActivity extends $${gameActivitySuperClass}$$ implements Surfac
 							memLog += "memstat: "+memstatusline + "\n";
 				    	}
 					}
-					//$${gameActivityMemStatAdditions}$$
+//$${gameActivityMemStatAdditions}$$
 					Log.debug(memLog);
 				}
-				synchronized(activityContext)
+				synchronized(_activity)
 				{
-					Activity localActivity = SimpleContextWrapper.getActivityOfContext(activityContext);
 					if (ProcessMemory > 0)
 					{
-						UsedMemory = ProcessMemory;
+						_activity.UsedMemory = ProcessMemory;
 					}
 					
-					if (memoryHandler != null && memstatLoggingInterval > 0)
+					if (_activity.memoryHandler != null && memstatLoggingInterval > 0)
 					{
-						memoryHandler.postDelayed(this, memstatLoggingInterval);
+						_activity.memoryHandler.postDelayed(this, memstatLoggingInterval);
 					}
 				}
 			}
 		};
-
 		memoryHandlerThread = new HandlerThread("MemoryUsageThread", android.os.Process.THREAD_PRIORITY_LOWEST);
 		memoryHandlerThread.start();
 		memoryHandler = new Handler(memoryHandlerThread.getLooper());
@@ -2682,11 +2754,8 @@ public class GameActivity extends $${gameActivitySuperClass}$$ implements Surfac
 		{
 			MarginLayoutParams params = new MarginLayoutParams(LayoutParams.WRAP_CONTENT, LayoutParams.WRAP_CONTENT);
 			params.setMargins(0, 0, 0, 0);
-			activityLayout = new LinearLayout(activityContext);
-			SimpleContextWrapper.getActivityOfContext(activityContext).setContentView(activityLayout, params);
-
-			// tell Android that we want volume controls to change the media volume, aka music
-			setVolumeControlStream(AudioManager.STREAM_MUSIC);
+			activityLayout = new LinearLayout(_activity);
+			_activity.setContentView(activityLayout, params);
 		}
 
 		//
@@ -2736,6 +2805,11 @@ public class GameActivity extends $${gameActivitySuperClass}$$ implements Surfac
 			}
 		}
 */
+		if (!isStandalone)
+		{
+			// tell Android that we want volume controls to change the media volume, aka music
+			setVolumeControlStream(AudioManager.STREAM_MUSIC);
+		}
 
 		// Look for Vulkan support if Nougat or later
 		if (ANDROID_BUILD_VERSION >= 24)
@@ -2786,7 +2860,7 @@ public class GameActivity extends $${gameActivitySuperClass}$$ implements Surfac
 		}
 
 		// is this a native landscape device (tablet, tv)?
-		if ( getDeviceDefaultOrientation() == Configuration.ORIENTATION_LANDSCAPE )
+		if (!isStandalone && getDeviceDefaultOrientation() == Configuration.ORIENTATION_LANDSCAPE )
 		{
 			boolean bForceLandscape = false;
 
@@ -2834,7 +2908,7 @@ public class GameActivity extends $${gameActivitySuperClass}$$ implements Surfac
 			if (bForceLandscape)
 			{
 				Log.debug( "Setting screen orientation to landscape because we have detected landscape device" );
-				SimpleContextWrapper.getActivityOfContext(activityContext).setRequestedOrientation( android.content.pm.ActivityInfo.SCREEN_ORIENTATION_LANDSCAPE );
+				_activity.setRequestedOrientation( android.content.pm.ActivityInfo.SCREEN_ORIENTATION_LANDSCAPE );
 			}
 		}
 
@@ -3073,10 +3147,10 @@ public class GameActivity extends $${gameActivitySuperClass}$$ implements Surfac
 
 		AlertDialog.Builder builder;
 
-		consoleInputBox = new ConsoleKeyboardInput(activityContext);
+		consoleInputBox = new ConsoleKeyboardInput(getThisActivity());
 		consoleInputBox.setInputType(0x00080001); // TYPE_CLASS_TEXT | TYPE_TEXT_FLAG_NO_SUGGESTIONS);
 		
-		RelativeLayout consoleInputLayout = new RelativeLayout(activityContext);
+		RelativeLayout consoleInputLayout = new RelativeLayout(getThisActivity());
 		{ 
 			android.widget.ImageButton consoleHistoryButton = consoleInputBox.getButton();
 			consoleInputBox.setId(1);
@@ -3094,8 +3168,8 @@ public class GameActivity extends $${gameActivitySuperClass}$$ implements Surfac
 		}
 
 		// Spinner with Quick Stat Commands
-		consoleSpinner = new Spinner(activityContext);
-		ArrayAdapter<String> adapter = new ArrayAdapter<String>(activityContext, android.R.layout.simple_spinner_item, CONSOLE_SPINNER_ITEMS);
+		consoleSpinner = new Spinner(getThisActivity());
+		ArrayAdapter<String> adapter = new ArrayAdapter<String>(getThisActivity(), android.R.layout.simple_spinner_item, CONSOLE_SPINNER_ITEMS);
 		adapter.setDropDownViewResource(android.R.layout.simple_spinner_dropdown_item);
 		consoleSpinner.setAdapter(adapter);
 		consoleSpinner.setOnItemSelectedListener(new AdapterView.OnItemSelectedListener() {
@@ -3149,7 +3223,7 @@ public class GameActivity extends $${gameActivitySuperClass}$$ implements Surfac
 		});
 		consoleAlert = builder.create();
 
-		virtualKeyboardInputBox = new EditText(activityContext);
+		virtualKeyboardInputBox = new EditText(getThisActivity());
 		virtualKeyboardInputBox.addTextChangedListener(new TextWatcher() {
 			@Override
 			public void beforeTextChanged(CharSequence charSequence, int start, int count, int after) {
@@ -3166,7 +3240,7 @@ public class GameActivity extends $${gameActivitySuperClass}$$ implements Surfac
 			}
 		});
 
-		builder = new AlertDialog.Builder(activityContext);
+		builder = new AlertDialog.Builder(getThisActivity());
 		builder.setTitle("")
 		.setView(virtualKeyboardInputBox)
 		.setPositiveButton("Ok", new DialogInterface.OnClickListener() {
@@ -3211,7 +3285,7 @@ public class GameActivity extends $${gameActivitySuperClass}$$ implements Surfac
 		// check for OBB file present if we don't have all the files and don't need to verify
 		if (!isStandalone && !HasAllFiles && !VerifyOBBOnStartUp)
 		{
-			HasAllFiles = DownloadShim.expansionFilesDelivered(this, VersionCode);
+			HasAllFiles = DownloadShim.expansionFilesDelivered(getThisActivity(), VersionCode);
 		}
 
 		containerFrameLayout = new FrameLayout(activityContext);
@@ -3223,7 +3297,7 @@ public class GameActivity extends $${gameActivitySuperClass}$$ implements Surfac
 
 			// Need to create our surface view here regardless of if we are going to end up using it
 			getWindow().takeSurface(null);
-			MySurfaceView = new SurfaceView(this);
+			MySurfaceView = new SurfaceView(getThisActivity());
 			MySurfaceView.setBackgroundColor(Color.TRANSPARENT);
 			MySurfaceView.getHolder().addCallback(this);
 			containerFrameLayout.addView(MySurfaceView);
@@ -3405,7 +3479,7 @@ public class GameActivity extends $${gameActivitySuperClass}$$ implements Surfac
 
 	public void AndroidThunkJava_ResizeKeyboard(boolean bShow, Rect keyboardRect, Rect visibleRect)
 	{
-		//$${gameActivityResizeKeyboardAdditions}$$
+//$${gameActivityResizeKeyboardAdditions}$$
 	}
 
 	@Override
@@ -3431,7 +3505,7 @@ public class GameActivity extends $${gameActivitySuperClass}$$ implements Surfac
 
 	public void restoreTransparentBars()
 	{
-		if(android.os.Build.VERSION.SDK_INT >= 19) {
+		if(!isStandalone && android.os.Build.VERSION.SDK_INT >= 19) {
 			try {
 				View decorView = getWindow().getDecorView(); 
 
@@ -3579,7 +3653,7 @@ public class GameActivity extends $${gameActivitySuperClass}$$ implements Surfac
 		bIsActivityPaused.set(false);
 
 		// start memory reporter timer (runs every "memstatLoggingInterval" seconds)
-		synchronized(this)
+		synchronized(getThisActivity())
 		{
 			if (memoryHandler == null && memoryRunnable != null && memstatLoggingInterval > 0)
 			{
@@ -3631,11 +3705,11 @@ public class GameActivity extends $${gameActivitySuperClass}$$ implements Surfac
 		// restore screep capture state
 		SetDisableScreenCaptureInternal(AndroidThunkJava_IsScreenCaptureDisabled());
 
-		if (bForceGameEndWithError || bForceGameEndWithWarning)
+		if (!isStandalone && (bForceGameEndWithError || bForceGameEndWithWarning))
 		{
 			final boolean bWarning = bForceGameEndWithWarning;
 
-			activityContext.runOnUiThread(new Runnable()
+			_activity.runOnUiThread(new Runnable()
 			{
 				@Override
 				public void run()
@@ -3645,7 +3719,7 @@ public class GameActivity extends $${gameActivitySuperClass}$$ implements Surfac
 						return;
 					}
 
-					AlertDialog.Builder dialog = new AlertDialog.Builder(activityContext);
+					AlertDialog.Builder dialog = new AlertDialog.Builder(_activity);
 					dialog.setCancelable(false);
 					dialog.setTitle(ForceExitCaption);
 					dialog.setMessage(ForceExitMessage);
@@ -3727,7 +3801,7 @@ public class GameActivity extends $${gameActivitySuperClass}$$ implements Surfac
 		bIsActivityPaused.set(true);
 
 		// stop memory reporter timer if running
-		synchronized(this)
+		synchronized(_activity)
 		{
 			if (memoryHandler != null)
 			{
@@ -3753,7 +3827,7 @@ public class GameActivity extends $${gameActivitySuperClass}$$ implements Surfac
 		{
 			//	If an AlertDialog is showing when the application is paused, it can cause our main window to be terminated
 			//	Hide the dialog here. It will be shown again via AndroidThunkJava_ShowHiddenAlertDialog called from native code
-			activityContext.runOnUiThread(new Runnable()
+			_activity.runOnUiThread(new Runnable()
 			{
 				public void run()
 				{
@@ -3995,7 +4069,7 @@ public class GameActivity extends $${gameActivitySuperClass}$$ implements Surfac
 	public void onConfigurationChanged(Configuration newConfig)
 	{
 		super.onConfigurationChanged(newConfig);
-		//$${gameActivityonConfigurationChangedAdditions}$$
+//$${gameActivityonConfigurationChangedAdditions}$$
 		// forward the orientation
 		nativeOnOrientationChanged(newConfig.orientation);
 	}
@@ -4026,7 +4100,7 @@ public class GameActivity extends $${gameActivitySuperClass}$$ implements Surfac
 			Log.debug("==============> [JAVA] AndroidThunkJava_ShowHiddenAlertDialog() - Showing " + CurrentDialogType);
 		
 			//	If an AlertDialog was showing onPause and we hid it, show it again
-			activityContext.runOnUiThread(new Runnable()
+			_activity.runOnUiThread(new Runnable()
 			{
 				public void run()
 				{
@@ -4062,25 +4136,25 @@ public class GameActivity extends $${gameActivitySuperClass}$$ implements Surfac
 		bKeepScreenOn = Enable;
 		if (Enable)
 		{
-			activityContext.runOnUiThread(new Runnable()
+			_activity.runOnUiThread(new Runnable()
 			{
 				@Override
 				public void run()
 				{
 					Log.debug("==============> [JAVA] AndroidThunkJava_KeepScreenOn(true) - Disabled screen saver");
-					getWindow().addFlags(WindowManager.LayoutParams.FLAG_KEEP_SCREEN_ON);
+					_activity.getWindow().addFlags(WindowManager.LayoutParams.FLAG_KEEP_SCREEN_ON);
 				}
 			});
 		}
 		else
 		{
-			activityContext.runOnUiThread(new Runnable()
+			_activity.runOnUiThread(new Runnable()
 			{
 				@Override
 				public void run()
 				{
 					Log.debug("==============> [JAVA] AndroidThunkJava_KeepScreenOn(false) - Enabled screen saver");
-					getWindow().clearFlags(WindowManager.LayoutParams.FLAG_KEEP_SCREEN_ON);
+					_activity.getWindow().clearFlags(WindowManager.LayoutParams.FLAG_KEEP_SCREEN_ON);
 				}
 			});
 		}
@@ -4100,25 +4174,25 @@ public class GameActivity extends $${gameActivitySuperClass}$$ implements Surfac
 
 		if (bDisable)
 		{
-			activityContext.runOnUiThread(new Runnable()
+			_activity.runOnUiThread(new Runnable()
 			{
 				@Override
 				public void run()
 				{
 					Log.debug("==============> [JAVA] AndroidThunkJava_DisableScreenCapture(true) - Disabled screen captures");
-					getWindow().addFlags(WindowManager.LayoutParams.FLAG_SECURE);
+					_activity.getWindow().addFlags(WindowManager.LayoutParams.FLAG_SECURE);
 				}
 			});
 		}
 		else
 		{
-			activityContext.runOnUiThread(new Runnable()
+			_activity.runOnUiThread(new Runnable()
 			{
 				@Override
 				public void run()
 				{
 					Log.debug("==============> [JAVA] AndroidThunkJava_DisableScreenCapture(false) - Enabled screen captures");
-					getWindow().clearFlags(WindowManager.LayoutParams.FLAG_SECURE);
+					_activity.getWindow().clearFlags(WindowManager.LayoutParams.FLAG_SECURE);
 				}
 			});
 		}
@@ -4161,7 +4235,7 @@ public class GameActivity extends $${gameActivitySuperClass}$$ implements Surfac
 	{
 		if (deviceVibrator != null)
 		{
-			activityContext.runOnUiThread(new VibrateRunnable(Intensity, (long)Duration, deviceVibrator));
+			_activity.runOnUiThread(new VibrateRunnable(Intensity, (long)Duration, deviceVibrator));
 		}
 	}
 
@@ -4178,7 +4252,7 @@ public class GameActivity extends $${gameActivitySuperClass}$$ implements Surfac
 		consoleInputBox.setHistoryEnd();
 
 		consoleAlert.setMessage("[Available texture formats: " + Formats + "]");
-		activityContext.runOnUiThread(new Runnable()
+		_activity.runOnUiThread(new Runnable()
 		{
 			public void run()
 			{
@@ -4201,7 +4275,7 @@ public class GameActivity extends $${gameActivitySuperClass}$$ implements Surfac
 			return;
 		}
 
-		activityContext.runOnUiThread(new Runnable()
+		_activity.runOnUiThread(new Runnable()
 		{
 			public void run()
 			{
@@ -4229,7 +4303,7 @@ public class GameActivity extends $${gameActivitySuperClass}$$ implements Surfac
 		final String uiLabel = inLabel;
 		final String uiContents = inContents;
 
-		activityContext.runOnUiThread(new Runnable()
+		_activity.runOnUiThread(new Runnable()
 		{
 			public void run()
 			{
@@ -4520,7 +4594,7 @@ public class GameActivity extends $${gameActivitySuperClass}$$ implements Surfac
 	public String AndroidThunkJava_GetFunnelId()
 	{
 		String funnelId = "";
-		//$${gameActivityGetFunnelIdAdditions}$$
+//$${gameActivityGetFunnelIdAdditions}$$
 		return funnelId;
 	}
 
@@ -4528,7 +4602,7 @@ public class GameActivity extends $${gameActivitySuperClass}$$ implements Surfac
 	{
 		String loginId = null;
 
-	//$${gameActivityGetLoginIdAdditions}$$
+//$${gameActivityGetLoginIdAdditions}$$
 
 		// check for existing GUID fallback in internal
 		File file = new File(InternalFilesDir + "login-identifier.txt");
@@ -4676,12 +4750,12 @@ public class GameActivity extends $${gameActivitySuperClass}$$ implements Surfac
 	public void AndroidThunkJava_SetOrientation(int value)
 	{
 		final int orientation = value;
-		activityContext.runOnUiThread(new Runnable()
+		_activity.runOnUiThread(new Runnable()
 		{
 			@Override
 			public void run()
 			{
-				SimpleContextWrapper.getActivityOfContext(activityContext).setRequestedOrientation( orientation );
+				_activity.setRequestedOrientation( orientation );
 			}
 		});
 	}
@@ -4705,7 +4779,7 @@ public class GameActivity extends $${gameActivitySuperClass}$$ implements Surfac
 	// call back into native code from the Java UI thread, initializing any available VR HMD modules
 	public void AndroidThunkJava_InitHMDs()
 	{
-		activityContext.runOnUiThread(new Runnable()
+		_activity.runOnUiThread(new Runnable()
 		{
 			@Override
 			public void run()
@@ -4743,7 +4817,7 @@ public class GameActivity extends $${gameActivitySuperClass}$$ implements Surfac
 
 	public boolean AndroidThunkJava_IsMusicActive()
 	{
-		AudioManager audioManager = (AudioManager)getSystemService(Context.AUDIO_SERVICE);
+		AudioManager audioManager = (AudioManager)getThisActivity().getSystemService(Context.AUDIO_SERVICE);
 		return audioManager.isMusicActive();
 	}
 	
@@ -4765,7 +4839,7 @@ public class GameActivity extends $${gameActivitySuperClass}$$ implements Surfac
 		{
 			bTriggeredQuery = true;
 
-			activityContext.runOnUiThread(new Runnable()
+			_activity.runOnUiThread(new Runnable()
 			{
 				@Override
 				public void run()
@@ -4928,18 +5002,22 @@ public class GameActivity extends $${gameActivitySuperClass}$$ implements Surfac
 
 	public void AndroidThunkJava_ShowProgressDialog(boolean bShow, String message, boolean bHorizontal, int maxValue)
 	{
+		if (isStandalone)
+		{
+			return;
+		}
 		final boolean bShowFinal = bShow;
 		final String messageFinal = message;
 		final boolean bHorizontalFinal = bHorizontal;
 		final int maxValueFinal = maxValue;
-		activityContext.runOnUiThread(new Runnable()
+		_activity.runOnUiThread(new Runnable()
 		{
 			@Override
 			public void run()
 			{
 				if (bShowFinal)
 				{
-					mProgressDialog = new ProgressDialog(activityContext);
+					mProgressDialog = new ProgressDialog(_activity);
 					mProgressDialog.setMessage(messageFinal);
 					mProgressDialog.setProgress(0);
 					mProgressDialog.setMax(maxValueFinal);
@@ -4965,7 +5043,7 @@ public class GameActivity extends $${gameActivitySuperClass}$$ implements Surfac
 			return;
 		}
 		final int valueFinal = value;
-		activityContext.runOnUiThread(new Runnable()
+		_activity.runOnUiThread(new Runnable()
 		{
 			@Override
 			public void run()
@@ -5483,7 +5561,7 @@ public class GameActivity extends $${gameActivitySuperClass}$$ implements Surfac
 	{
 		if (vibrator != null)
 		{
-			activityContext.runOnUiThread(new VibrateRunnable(intensity, duration, vibrator));
+			_activity.runOnUiThread(new VibrateRunnable(intensity, duration, vibrator));
 			return true;
 		}
 		return false;
@@ -5538,7 +5616,7 @@ public class GameActivity extends $${gameActivitySuperClass}$$ implements Surfac
 			DesiredHolderHeight > 0 &&
 			MySurfaceView != null)
 		{
-			activityContext.runOnUiThread(new Runnable()
+			_activity.runOnUiThread(new Runnable()
 			{
 				@Override
 				public void run()
@@ -5569,7 +5647,7 @@ public class GameActivity extends $${gameActivitySuperClass}$$ implements Surfac
 
 		if(bUseSurfaceView && MySurfaceView != null)
 		{
-			activityContext.runOnUiThread(new Runnable()
+			_activity.runOnUiThread(new Runnable()
 			{
 				@Override
 				public void run()
@@ -5650,7 +5728,7 @@ public class GameActivity extends $${gameActivitySuperClass}$$ implements Surfac
 		if (ANDROID_BUILD_VERSION >= 30)
 		{
 			final int FinalRefreshRate = RefreshRate;
-			activityContext.runOnUiThread(new Runnable()
+			_activity.runOnUiThread(new Runnable()
 			{
 				@Override
 				public void run()
@@ -5689,7 +5767,7 @@ public class GameActivity extends $${gameActivitySuperClass}$$ implements Surfac
 					{
 						if (!isStandalone)
 						{
-							activityContext.runOnUiThread(new Runnable()
+							_activity.runOnUiThread(new Runnable()
 							{
 								public void run()
 								{
@@ -5869,9 +5947,6 @@ public class GameActivity extends $${gameActivitySuperClass}$$ implements Surfac
 			return 0;
 		}
 
-		assert(activityContext != null);
-
-
 		switch(key)
 		{
 			case "unreal.http.proxy.proxyPort":
@@ -5897,9 +5972,9 @@ public class GameActivity extends $${gameActivitySuperClass}$$ implements Surfac
 			case "unreal.getUsedMemory":
 			{
 				int ProcessMemory = 0;
-				synchronized(activityContext)
+				synchronized(_activity)
 				{
-					ProcessMemory = UsedMemory;
+					ProcessMemory = _activity.UsedMemory;
 				}
 				return ProcessMemory;
 			}
@@ -5938,7 +6013,7 @@ public class GameActivity extends $${gameActivitySuperClass}$$ implements Surfac
 		{
 			if (android.os.Build.VERSION.SDK_INT >= 21)
 			{
-				WindowManager windowManager = getWindowManager();
+				WindowManager windowManager = (WindowManager)getSystemService(WINDOW_SERVICE);
 				android.view.Display display = windowManager.getDefaultDisplay();
 				return display.getPresentationDeadlineNanos();
 			}
@@ -5949,7 +6024,7 @@ public class GameActivity extends $${gameActivitySuperClass}$$ implements Surfac
 		{
 			if (android.os.Build.VERSION.SDK_INT >= 21)
 			{
-				WindowManager windowManager = getWindowManager();
+				WindowManager windowManager = (WindowManager)getSystemService(WINDOW_SERVICE);
 				android.view.Display display = windowManager.getDefaultDisplay();
 				return display.getAppVsyncOffsetNanos();
 			}
@@ -5962,7 +6037,7 @@ public class GameActivity extends $${gameActivitySuperClass}$$ implements Surfac
 	{
 		if (key.equals("unreal.display.getRefreshRate"))
 		{
-			WindowManager windowManager = getWindowManager();
+			WindowManager windowManager = (WindowManager)getSystemService(WINDOW_SERVICE);
 			android.view.Display display = windowManager.getDefaultDisplay();
 			float rate = display.getRefreshRate();
 			// make sure it is sane.. some drivers returned crazy values
@@ -6197,13 +6272,13 @@ public class GameActivity extends $${gameActivitySuperClass}$$ implements Surfac
 			}
 			if (!isStandalone)
 			{
-				activityContext.runOnUiThread(new Runnable()
+				_activity.runOnUiThread(new Runnable()
 				{
 					public void run()
 					{
 						try
 						{
-							android.view.Window ActivityWindow = getWindow();
+							android.view.Window ActivityWindow = _activity.getWindow();
 							Method m = android.view.Window.class.getMethod("setSustainedPerformanceMode",Boolean.TYPE);
 							m.invoke(ActivityWindow, bEnable);
 						}
@@ -6222,7 +6297,7 @@ public class GameActivity extends $${gameActivitySuperClass}$$ implements Surfac
 	}
 
 	//console input text input class with history - custom EditText
-	public class ConsoleKeyboardInput extends EditText
+	public class ConsoleKeyboardInput extends androidx.appcompat.widget.AppCompatEditText
 	{
 		private android.widget.ImageButton historyButton;
 		private ArrayList<String> historyList;
@@ -6308,7 +6383,7 @@ public class GameActivity extends $${gameActivitySuperClass}$$ implements Surfac
 	}
 	
 	//virtual keyboard input class - custom EditText
-	public class VirtualKeyboardInput extends EditText 
+	public class VirtualKeyboardInput extends androidx.appcompat.widget.AppCompatEditText 
 	{
 		public VirtualKeyboardInput(Context context, AttributeSet attrs, int defStyle) 
 		{
@@ -6341,6 +6416,7 @@ public class GameActivity extends $${gameActivitySuperClass}$$ implements Surfac
 		@Override   
 		protected void onSelectionChanged(int selStart, int selEnd) 
 		{
+			super.onSelectionChanged(selStart, selEnd);
 			nativeVirtualKeyboardSendSelection(selStart, selEnd);
 		}
 
@@ -6799,7 +6875,7 @@ public class GameActivity extends $${gameActivitySuperClass}$$ implements Surfac
 		};
 		Log.debug("app restart : " + RestartExtra);
 //$${gameActivityOnRestartApplicationAdditions}$$
-		activityContext.runOnUiThread(restartRunnableUIThread);
+		runOnUiThread(restartRunnableUIThread);
 
 		// wait for exit.
 		while(true)
@@ -7023,7 +7099,7 @@ public class GameActivity extends $${gameActivitySuperClass}$$ implements Surfac
 
 	public native void nativeInitHMDs();
 
-	public native static void nativeResumeMainInit();
+	public native void nativeResumeMainInit();
 
 	public native void nativeOnActivityResult(GameActivity activity, int requestCode, int resultCode, Intent data);
 
diff --git a/Engine/Build/Android/Java/src/com/epicgames/unreal/MessageBox01.java b/Engine/Build/Android/Java/src/com/epicgames/unreal/MessageBox01.java
--- a/Engine/Build/Android/Java/src/com/epicgames/unreal/MessageBox01.java
+++ b/Engine/Build/Android/Java/src/com/epicgames/unreal/MessageBox01.java
@@ -74,6 +74,24 @@ class MessageBox01
 	protected void createAlert()
 	{
 		dismissAlert();
+
+		if (GameApplication.getAppContext() == null)
+		{
+			GameActivity.Log.error("GameApplication.getAppContext() is NULL and cannot create Alert for Caption=" + Caption + ", text=" + Text);
+			return;
+		}
+		if (!GameActivity.isValidGameActivity() || GameActivity.Get().getApplicationContext() == null || GameActivity.getCurrentActivityContext() == null)
+		{
+			GameActivity.Log.error("GameActivity Application Context is not VALID yet! Cannot create Alert for Caption=" + Caption + ", text=" + Text);
+			return;
+		}
+		
+		if (GameActivity.isStandaloneMode())
+		{
+			GameActivity.Log.warn("AlertBox messages are not allowed in Standalone Mode (ASIS)! Caption=" + Caption + ", text=" + Text);
+			return;
+		}
+		
 		ButtonIdToIndex.clear();
 		Builder = new AlertDialog.Builder(GameActivity.Get());
 		Builder.setTitle(Caption);
diff --git a/Engine/Build/Android/Java/src/com/epicgames/unreal/SimpleContextWrapper.java b/Engine/Build/Android/Java/src/com/epicgames/unreal/SimpleContextWrapper.java
deleted file mode 100644
