Set user agent on mobile WebView browsers

#tests Checked user agent behaviour on Win64, Android and iOS

diff --git a/Engine/Build/Android/Java/src/com/epicgames/unreal/WebViewControl.java b/Engine/Build/Android/Java/src/com/epicgames/unreal/WebViewControl.java
--- a/Engine/Build/Android/Java/src/com/epicgames/unreal/WebViewControl.java
+++ b/Engine/Build/Android/Java/src/com/epicgames/unreal/WebViewControl.java
@@ -1,51 +1,35 @@
 // Copyright Epic Games, Inc. All Rights Reserved.
 package com.epicgames.unreal;
 
-import android.os.Build;
-import android.os.Bundle;
 import android.content.Context;
-
+import android.graphics.Bitmap;
+import android.graphics.Bitmap.Config;
+import android.graphics.Canvas;
+import android.graphics.Color;
+import android.opengl.EGL14;
+import android.opengl.EGLConfig;
+import android.opengl.EGLContext;
+import android.opengl.EGLDisplay;
+import android.opengl.EGLSurface;
+import android.opengl.GLES20;
+import android.util.AttributeSet;
+import android.util.Log;
+import android.os.SystemClock;
+import android.os.Message;
+import android.view.KeyEvent;
+import android.view.MotionEvent;
 import android.view.View;
 import android.view.ViewGroup;
-import android.view.KeyEvent;
 import android.webkit.ConsoleMessage;
-import android.webkit.WebView;
-import android.webkit.WebViewClient;
-import android.webkit.WebChromeClient;
-import android.webkit.WebSettings;
-import android.webkit.WebResourceRequest;
-import android.webkit.WebResourceResponse;
 import android.webkit.JsResult;
 import android.webkit.JsPromptResult;
-import android.graphics.Bitmap;
-import android.graphics.Bitmap.Config;
-import android.graphics.Paint;
-import android.graphics.Color;
 import android.webkit.WebBackForwardList;
+import android.webkit.WebChromeClient;
+import android.webkit.WebResourceResponse;
+import android.webkit.WebSettings;
+import android.webkit.WebView;
+import android.webkit.WebViewClient;
 import java.io.ByteArrayInputStream;
-import java.util.concurrent.ConcurrentHashMap;
-import java.util.concurrent.Semaphore;
-
-import android.util.Log;
-import android.opengl.*;
-import android.view.Surface;	
-import android.graphics.Canvas;
-import android.graphics.SurfaceTexture;
-
-import android.opengl.*;
-import android.widget.LinearLayout;
-import android.util.AttributeSet;
-import java.nio.ByteBuffer;
-import java.nio.ByteOrder;
-import java.nio.FloatBuffer;
-
-import android.opengl.GLES11Ext;
-import android.opengl.GLES20;
-import android.opengl.Matrix;
-import android.os.SystemClock;
-import android.view.MotionEvent;
-import android.os.Message;
-
 
 
 // Simple layout to apply absolute positioning for the WebView
@@ -98,7 +82,7 @@ class WebViewControl
 		public float VOffset;
 	}
 
-	public WebViewControl(long inNativePtr, int width, int height, boolean swizzlePixels, boolean vulkanRenderer, final boolean bEnableRemoteDebugging, final boolean bUseTransparency, final boolean bEnableDomStorage, final boolean bShouldUseBitmapRender)
+	public WebViewControl(long inNativePtr, int width, int height, boolean swizzlePixels, boolean vulkanRenderer, final boolean bEnableRemoteDebugging, final boolean bUseTransparency, final boolean bEnableDomStorage, final boolean bShouldUseBitmapRender, final String userAgentApplication)
 	{
 		final WebViewControl w = this;
 
@@ -128,16 +112,14 @@ class WebViewControl
 			@Override
 			public void run()
 			{
-				// enable remote debugging if requested and supported by the current platform
-				if (Build.VERSION.SDK_INT >= Build.VERSION_CODES.KITKAT)
-				{
-					WebView.setWebContentsDebuggingEnabled(bEnableRemoteDebugging && !GameActivity.Get().nativeIsShippingBuild());
-				}
+				WebView.setWebContentsDebuggingEnabled(bEnableRemoteDebugging && !GameActivity.Get().nativeIsShippingBuild());
 				
 				// create the WebView
 				webView = new GLWebView(GameActivity.Get());
 				webView.setWebViewClient(new ViewClient());
 				webView.setWebChromeClient(new ChromeClient());
+				String UserAgent = WebSettings.getDefaultUserAgent(GameActivity.Get()) + " " + userAgentApplication;
+				webView.getSettings().setUserAgentString(UserAgent);
 				webView.getSettings().setJavaScriptEnabled(true);
 				webView.getSettings().setJavaScriptCanOpenWindowsAutomatically(true);
 				webView.getSettings().setAllowFileAccess( true );
@@ -146,10 +128,7 @@ class WebViewControl
 				webView.getSettings().setAllowUniversalAccessFromFileURLs(true);
 				webView.getSettings().setSupportMultipleWindows(true);
 				webView.getSettings().setDomStorageEnabled(bEnableDomStorage);
-				if (Build.VERSION.SDK_INT >= Build.VERSION_CODES.LOLLIPOP)
-				{
-					webView.getSettings().setMixedContentMode(0); // 0 = MIXED_CONTENT_ALWAYS_ALLOW
-				}
+				webView.getSettings().setMixedContentMode(WebSettings.MIXED_CONTENT_ALWAYS_ALLOW);
 
 				webView.getSettings().setCacheMode( WebSettings.LOAD_DEFAULT );
 				webView.getSettings().setLoadWithOverviewMode(true);
@@ -254,15 +233,7 @@ class WebViewControl
 			{
 				if(webView != null)
 				{
-					if (Build.VERSION.SDK_INT >= Build.VERSION_CODES.KITKAT)
-					{
-						webView.evaluateJavascript(script, null);
-					}
-					else
-					{
-						// This is executed directly here instead of setting NextURL as otherwise calling ExecuteJavascript would only be possible once per tick.
-						webView.loadUrl("javascript:"+script);
-					}
+					webView.evaluateJavascript(script, null);
 				}
 			}
 		});
@@ -832,7 +803,7 @@ class WebViewControl
 				if (RendererString.contains("Adreno (TM) "))
 				{
 					int AdrenoVersion = Integer.parseInt(RendererString.substring(12));
-					if (AdrenoVersion < 400 || android.os.Build.VERSION.SDK_INT < 22)
+					if (AdrenoVersion < 400)
 					{
 						GameActivity.Log.debug("WebViewControl: disabled shared GL context on " + RendererString);
 						mUseOwnContext = false;
diff --git a/Engine/Source/Runtime/WebBrowser/Private/Android/AndroidJavaWebBrowser.cpp b/Engine/Source/Runtime/WebBrowser/Private/Android/AndroidJavaWebBrowser.cpp
--- a/Engine/Source/Runtime/WebBrowser/Private/Android/AndroidJavaWebBrowser.cpp
+++ b/Engine/Source/Runtime/WebBrowser/Private/Android/AndroidJavaWebBrowser.cpp
@@ -33,8 +33,8 @@ static jfieldID FindField(JNIEnv* JEnv, jclass Class, const ANSICHAR* FieldName,
 #define MOTIONEVENT_ACTION_MOVE		2
 
 FJavaAndroidWebBrowser::FJavaAndroidWebBrowser(bool swizzlePixels, bool vulkanRenderer, int32 width, int32 height,
-	jlong widgetPtr, bool bEnableRemoteDebugging, bool bUseTransparency, bool bEnableDomStorage, bool bShouldUseBitmapRender)
-	: FJavaClassObject(GetClassName(), "(JIIZZZZZZ)V", widgetPtr, width, height,  swizzlePixels, vulkanRenderer, bEnableRemoteDebugging, bUseTransparency, bEnableDomStorage, bShouldUseBitmapRender)
+	jlong widgetPtr, bool bEnableRemoteDebugging, bool bUseTransparency, bool bEnableDomStorage, bool bShouldUseBitmapRender, const FString& UserAgentApplication)
+	: FJavaClassObject(GetClassName(), "(JIIZZZZZZLjava/lang/String;)V", widgetPtr, width, height,  swizzlePixels, vulkanRenderer, bEnableRemoteDebugging, bUseTransparency, bEnableDomStorage, bShouldUseBitmapRender, *FJavaHelper::ToJavaString(AndroidJavaEnv::GetJavaEnv(), UserAgentApplication))
 	, ReleaseMethod(GetClassMethod("release", "()V"))
 	, GetVideoLastFrameBitmapMethod(GetClassMethod("getVideoLastFrameBitmap", "()Lcom/epicgames/unreal/WebViewControl$FrameUpdateInfo;"))
 	, GetVideoLastFrameDataMethod(GetClassMethod("getVideoLastFrameData", "()Lcom/epicgames/unreal/WebViewControl$FrameUpdateInfo;"))
diff --git a/Engine/Source/Runtime/WebBrowser/Private/Android/AndroidJavaWebBrowser.h b/Engine/Source/Runtime/WebBrowser/Private/Android/AndroidJavaWebBrowser.h
--- a/Engine/Source/Runtime/WebBrowser/Private/Android/AndroidJavaWebBrowser.h
+++ b/Engine/Source/Runtime/WebBrowser/Private/Android/AndroidJavaWebBrowser.h
@@ -15,7 +15,7 @@
 class FJavaAndroidWebBrowser : public FJavaClassObject
 {
 public:
-	FJavaAndroidWebBrowser(bool swizzlePixels, bool vulkanRenderer, int32 width, int32 height, jlong widgetPtr, bool bEnableRemoteDebugging, bool bUseTransparency, bool bEnableDomStorage, bool bShouldUseBitmapRender);
+	FJavaAndroidWebBrowser(bool swizzlePixels, bool vulkanRenderer, int32 width, int32 height, jlong widgetPtr, bool bEnableRemoteDebugging, bool bUseTransparency, bool bEnableDomStorage, bool bShouldUseBitmapRender, const FString& UserAgentApplication);
 	virtual ~FJavaAndroidWebBrowser();
 	void Release();
 	bool GetVideoLastFrameBitmap(void* outPixels, int64 outCount);
diff --git a/Engine/Source/Runtime/WebBrowser/Private/Android/AndroidWebBrowserWidget.cpp b/Engine/Source/Runtime/WebBrowser/Private/Android/AndroidWebBrowserWidget.cpp
--- a/Engine/Source/Runtime/WebBrowser/Private/Android/AndroidWebBrowserWidget.cpp
+++ b/Engine/Source/Runtime/WebBrowser/Private/Android/AndroidWebBrowserWidget.cpp
@@ -128,7 +128,7 @@ void SAndroidWebBrowserWidget::Construct(const FArguments& Args)
 	FPlatformMisc::LowLevelOutputDebugStringf(TEXT("SAndroidWebBrowserWidget::Construct viewport=%d x %d"), viewportSize.X, viewportSize.Y);
 
 	JavaWebBrowser = MakeShared<FJavaAndroidWebBrowser, ESPMode::ThreadSafe>(false, FAndroidMisc::ShouldUseVulkan(), viewportSize.X, viewportSize.Y,
-		reinterpret_cast<jlong>(this), !(UE_BUILD_SHIPPING || UE_BUILD_TEST), Args._UseTransparency, bEnableDomStorage, bShouldUseBitmapRender);
+		reinterpret_cast<jlong>(this), !(UE_BUILD_SHIPPING || UE_BUILD_TEST), Args._UseTransparency, bEnableDomStorage, bShouldUseBitmapRender, Args._UserAgentApplication);
 
 	TextureSamplePool = new FWebBrowserTextureSamplePool();
 	WebBrowserTextureSamplesQueue = MakeShared<FWebBrowserTextureSampleQueue, ESPMode::ThreadSafe>();
diff --git a/Engine/Source/Runtime/WebBrowser/Private/Android/AndroidWebBrowserWidget.h b/Engine/Source/Runtime/WebBrowser/Private/Android/AndroidWebBrowserWidget.h
--- a/Engine/Source/Runtime/WebBrowser/Private/Android/AndroidWebBrowserWidget.h
+++ b/Engine/Source/Runtime/WebBrowser/Private/Android/AndroidWebBrowserWidget.h
@@ -35,6 +35,7 @@ class SAndroidWebBrowserWidget : public SViewport
 		SLATE_ARGUMENT(FString, InitialURL);
 		SLATE_ARGUMENT(bool, UseTransparency);
 		SLATE_ARGUMENT(TSharedPtr<FAndroidWebBrowserWindow>, WebBrowserWindow);
+		SLATE_ARGUMENT(FString, UserAgentApplication);
 
 	SLATE_END_ARGS()
 
diff --git a/Engine/Source/Runtime/WebBrowser/Private/Android/AndroidWebBrowserWindow.cpp b/Engine/Source/Runtime/WebBrowser/Private/Android/AndroidWebBrowserWindow.cpp
--- a/Engine/Source/Runtime/WebBrowser/Private/Android/AndroidWebBrowserWindow.cpp
+++ b/Engine/Source/Runtime/WebBrowser/Private/Android/AndroidWebBrowserWindow.cpp
@@ -23,9 +23,10 @@ namespace {
 
 }
 
-FAndroidWebBrowserWindow::FAndroidWebBrowserWindow(FString InUrl, TOptional<FString> InContentsToLoad, bool InShowErrorMessage, bool InThumbMouseButtonNavigation, bool InUseTransparency, bool bInJSBindingToLoweringEnabled)
+FAndroidWebBrowserWindow::FAndroidWebBrowserWindow(FString InUrl, TOptional<FString> InContentsToLoad, bool InShowErrorMessage, bool InThumbMouseButtonNavigation, bool InUseTransparency, bool bInJSBindingToLoweringEnabled, const FString& InUserAgentApplication)
 	: CurrentUrl(MoveTemp(InUrl))
 	, ContentsToLoad(MoveTemp(InContentsToLoad))
+	, UserAgentApplication(InUserAgentApplication)
 	, bUseTransparency(InUseTransparency)
 	, DocumentState(EWebBrowserDocumentState::NoDocument)
 	, ErrorCode(0)
@@ -75,6 +76,7 @@ TSharedRef<SWidget> FAndroidWebBrowserWindow::CreateWidget()
 		SNew(SAndroidWebBrowserWidget)
 		.UseTransparency(bUseTransparency)
 		.InitialURL(CurrentUrl)
+		.UserAgentApplication(UserAgentApplication)
 		.WebBrowserWindow(SharedThis(this));
 
 	BrowserWidget = BrowserWidgetRef;
diff --git a/Engine/Source/Runtime/WebBrowser/Private/Android/AndroidWebBrowserWindow.h b/Engine/Source/Runtime/WebBrowser/Private/Android/AndroidWebBrowserWindow.h
--- a/Engine/Source/Runtime/WebBrowser/Private/Android/AndroidWebBrowserWindow.h
+++ b/Engine/Source/Runtime/WebBrowser/Private/Android/AndroidWebBrowserWindow.h
@@ -37,7 +37,7 @@ private:
 	 * @param bUseTransparency Whether to enable transparency.
 	 * @param bJSBindingToLoweringEnabled Whether we ToLower all JavaScript member names.
 	 */
-	FAndroidWebBrowserWindow(FString Url, TOptional<FString> ContentsToLoad, bool bShowErrorMessage, bool bThumbMouseButtonNavigation, bool bUseTransparency, bool bJSBindingToLoweringEnabled);
+	FAndroidWebBrowserWindow(FString Url, TOptional<FString> ContentsToLoad, bool bShowErrorMessage, bool bThumbMouseButtonNavigation, bool bUseTransparency, bool bJSBindingToLoweringEnabled, const FString& InUserAgentApplication);
 
 	/**
 	 * Create the SWidget for this WebBrowserWindow
@@ -259,6 +259,9 @@ private:
 	/** Optional text to load as a web page. */
 	TOptional<FString> ContentsToLoad;
 
+	/** User Agent Application to report */
+	FString UserAgentApplication;
+
 	/** Whether to enable background transparency */
 	bool bUseTransparency;
 
diff --git a/Engine/Source/Runtime/WebBrowser/Private/IOS/IOSPlatformWebBrowser.cpp b/Engine/Source/Runtime/WebBrowser/Private/IOS/IOSPlatformWebBrowser.cpp
--- a/Engine/Source/Runtime/WebBrowser/Private/IOS/IOSPlatformWebBrowser.cpp
+++ b/Engine/Source/Runtime/WebBrowser/Private/IOS/IOSPlatformWebBrowser.cpp
@@ -26,6 +26,7 @@ class SIOSWebBrowserWidget : public SLeafWidget
 	SLATE_ARGUMENT(FString, InitialURL);
 	SLATE_ARGUMENT(bool, UseTransparency);
 	SLATE_ARGUMENT(TSharedPtr<FWebBrowserWindow>, WebBrowserWindow);
+	SLATE_ARGUMENT(FString, UserAgentApplication);
 
 	SLATE_END_ARGS()
 
@@ -44,7 +45,7 @@ class SIOSWebBrowserWidget : public SLeafWidget
 		check(bSupportsMetal);
 
 		WebViewWrapper = [IOSWebViewWrapper alloc];
-		[WebViewWrapper create : TSharedPtr<SIOSWebBrowserWidget>(this) useTransparency : Args._UseTransparency supportsMetal : bSupportsMetal supportsMetalMRT : bSupportsMetalMRT];
+		[WebViewWrapper create : TSharedPtr<SIOSWebBrowserWidget>(this) userAgentApplication : Args._UserAgentApplication.GetNSString() useTransparency : Args._UseTransparency supportsMetal : bSupportsMetal supportsMetalMRT : bSupportsMetalMRT];
 
 		WebBrowserWindowPtr = Args._WebBrowserWindow;
 		IsIOS3DBrowser = false;
@@ -454,7 +455,7 @@ private:
 @synthesize NextURL;
 @synthesize NextContent;
 
--(void)create:(TSharedPtr<SIOSWebBrowserWidget>)InWebBrowserWidget useTransparency : (bool)InUseTransparency
+-(void)create:(TSharedPtr<SIOSWebBrowserWidget>)InWebBrowserWidget userAgentApplication: (NSString*)UserAgentApplication useTransparency : (bool)InUseTransparency
 supportsMetal : (bool)InSupportsMetal supportsMetalMRT : (bool)InSupportsMetalMRT;
 {
 	WebBrowserWidget = InWebBrowserWidget;
@@ -476,6 +477,7 @@ supportsMetal : (bool)InSupportsMetal supportsMetalMRT : (bool)InSupportsMetalMR
 
 		WKWebViewConfiguration *theConfiguration = [[WKWebViewConfiguration alloc] init];
 		NSString* MessageHandlerName = [NSString stringWithFString : FMobileJSScripting::JSMessageHandler];
+		theConfiguration.applicationNameForUserAgent = UserAgentApplication;
 		[theConfiguration.userContentController addScriptMessageHandler:self name: MessageHandlerName];
 
 		WebView = [[WKWebView alloc]initWithFrame:CGRectMake(1, 1, 100, 100)  configuration : theConfiguration];
@@ -844,8 +846,9 @@ namespace {
 
 }
 
-FWebBrowserWindow::FWebBrowserWindow(FString InUrl, TOptional<FString> InContentsToLoad, bool InShowErrorMessage, bool InThumbMouseButtonNavigation, bool InUseTransparency, bool bInJSBindingToLoweringEnabled)
+FWebBrowserWindow::FWebBrowserWindow(FString InUrl, TOptional<FString> InContentsToLoad, bool InShowErrorMessage, bool InThumbMouseButtonNavigation, bool InUseTransparency, bool bInJSBindingToLoweringEnabled, FString InUserAgentApplication)
 	: CurrentUrl(MoveTemp(InUrl))
+	, UserAgentApplication(MoveTemp(InUserAgentApplication))
 	, ContentsToLoad(MoveTemp(InContentsToLoad))
 	, bUseTransparency(InUseTransparency)
 	, DocumentState(EWebBrowserDocumentState::NoDocument)
@@ -879,6 +882,7 @@ TSharedRef<SWidget> FWebBrowserWindow::CreateWidget()
 		SNew(SIOSWebBrowserWidget)
 		.UseTransparency(bUseTransparency)
 		.InitialURL(CurrentUrl)
+		.UserAgentApplication(UserAgentApplication)
 		.WebBrowserWindow(SharedThis(this));
 
 	BrowserWidget = BrowserWidgetRef;
diff --git a/Engine/Source/Runtime/WebBrowser/Private/IOS/IOSPlatformWebBrowser.h b/Engine/Source/Runtime/WebBrowser/Private/IOS/IOSPlatformWebBrowser.h
--- a/Engine/Source/Runtime/WebBrowser/Private/IOS/IOSPlatformWebBrowser.h
+++ b/Engine/Source/Runtime/WebBrowser/Private/IOS/IOSPlatformWebBrowser.h
@@ -48,7 +48,7 @@ class SWebBrowserView;
 @property(copy) NSString* NextContent;
 @property CGRect DesiredFrame;
 
--(void)create:(TSharedPtr<SIOSWebBrowserWidget>)InWebBrowserWidget useTransparency : (bool)InUseTransparency
+-(void)create:(TSharedPtr<SIOSWebBrowserWidget>)InWebBrowserWidget userAgentApplication: (NSString*)UserAgentApplication useTransparency : (bool)InUseTransparency
 supportsMetal : (bool)InSupportsMetal supportsMetalMRT : (bool)InSupportsMetalMRT;
 -(void)close;
 -(void)dealloc;
@@ -97,7 +97,7 @@ private:
 	* @param InContentsToLoad Optional string to load as a web page.
 	* @param InShowErrorMessage Whether to show an error message in case of loading errors.
 	*/
-	FWebBrowserWindow(FString InUrl, TOptional<FString> InContentsToLoad, bool ShowErrorMessage, bool bThumbMouseButtonNavigation, bool bUseTransparency, bool bInJSBindingToLoweringEnabled);
+	FWebBrowserWindow(FString InUrl, TOptional<FString> InContentsToLoad, bool ShowErrorMessage, bool bThumbMouseButtonNavigation, bool bUseTransparency, bool bInJSBindingToLoweringEnabled, FString InUserAgentApplication);
 
 	/**
 	 * Create the SWidget for this WebBrowserWindow
@@ -318,6 +318,9 @@ private:
 	/** Current Url of this window. */
 	FString CurrentUrl;
 
+	/** User-Agent Application to report. */
+	FString UserAgentApplication;
+
 	/** Optional text to load as a web page. */
 	TOptional<FString> ContentsToLoad;
 	
diff --git a/Engine/Source/Runtime/WebBrowser/Private/WebBrowserModule.cpp b/Engine/Source/Runtime/WebBrowser/Private/WebBrowserModule.cpp
--- a/Engine/Source/Runtime/WebBrowser/Private/WebBrowserModule.cpp
+++ b/Engine/Source/Runtime/WebBrowser/Private/WebBrowserModule.cpp
@@ -18,7 +18,7 @@ DEFINE_LOG_CATEGORY(LogWebBrowser);
 static FWebBrowserSingleton* WebBrowserSingleton = nullptr;
 
 FWebBrowserInitSettings::FWebBrowserInitSettings()
-	: ProductVersion(FString::Printf(TEXT("%s/%s UnrealEngine/%s Chrome/90.0.4430.212"), FApp::GetProjectName(), FApp::GetBuildVersion(), *FEngineVersion::Current().ToString()))
+	: ProductVersion(FString::Printf(TEXT("%s/%s UnrealEngine/%s"), FApp::GetProjectName(), FApp::GetBuildVersion(), *FEngineVersion::Current().ToString()))
 {
 }
 
diff --git a/Engine/Source/Runtime/WebBrowser/Private/WebBrowserSingleton.cpp b/Engine/Source/Runtime/WebBrowser/Private/WebBrowserSingleton.cpp
--- a/Engine/Source/Runtime/WebBrowser/Private/WebBrowserSingleton.cpp
+++ b/Engine/Source/Runtime/WebBrowser/Private/WebBrowserSingleton.cpp
@@ -225,6 +225,9 @@ FWebBrowserSingleton::FWebBrowserSingleton(const FWebBrowserInitSettings& WebBro
 	: WebBrowserWindowFactory(MakeShareable(new FWebBrowserWindowFactory()))
 #else
 	: WebBrowserWindowFactory(MakeShareable(new FNoWebBrowserWindowFactory()))
+#if PLATFORM_IOS || (PLATFORM_ANDROID && USE_ANDROID_JNI)
+	, UserAgentApplication(WebBrowserInitSettings.ProductVersion)
+#endif
 #endif
 	, bDevToolsShortcutEnabled(UE_BUILD_DEBUG)
 	, bJSBindingsToLoweringEnabled(true)
@@ -292,7 +295,8 @@ FWebBrowserSingleton::FWebBrowserSingleton(const FWebBrowserInitSettings& WebBro
 		CefString(&Settings.locale) = TCHAR_TO_WCHAR(*LocaleCode);
 
 		// Append engine version to the user agent string.
-		CefString(&Settings.user_agent_product) = TCHAR_TO_WCHAR(*WebBrowserInitSettings.ProductVersion);
+		FString UserAgentApplication = FString::Printf(TEXT("%s Chrome/%d.%d.%d.%d"), *WebBrowserInitSettings.ProductVersion, CHROME_VERSION_MAJOR, CHROME_VERSION_MINOR, CHROME_VERSION_BUILD, CHROME_VERSION_PATCH);
+		CefString(&Settings.user_agent_product) = TCHAR_TO_WCHAR(*UserAgentApplication);
 
 #if CEF3_DEFAULT_CACHE
 		// Enable on disk cache
@@ -622,7 +626,8 @@ TSharedPtr<IWebBrowserWindow> FWebBrowserSingleton::CreateBrowserWindow(const FC
 		WindowSettings.bShowErrorMessage,
 		WindowSettings.bThumbMouseButtonNavigation,
 		WindowSettings.bUseTransparency,
-		bJSBindingsToLoweringEnabled));
+		bJSBindingsToLoweringEnabled,
+		UserAgentApplication));
 
 	{
 		FScopeLock Lock(&WindowInterfacesCS);
@@ -637,7 +642,8 @@ TSharedPtr<IWebBrowserWindow> FWebBrowserSingleton::CreateBrowserWindow(const FC
 		WindowSettings.bShowErrorMessage, 
 		WindowSettings.bThumbMouseButtonNavigation, 
 		WindowSettings.bUseTransparency,
-		bJSBindingsToLoweringEnabled));
+		bJSBindingsToLoweringEnabled,
+		UserAgentApplication));
 
 	{
 		FScopeLock Lock(&WindowInterfacesCS);
diff --git a/Engine/Source/Runtime/WebBrowser/Private/WebBrowserSingleton.h b/Engine/Source/Runtime/WebBrowser/Private/WebBrowserSingleton.h
--- a/Engine/Source/Runtime/WebBrowser/Private/WebBrowserSingleton.h
+++ b/Engine/Source/Runtime/WebBrowser/Private/WebBrowserSingleton.h
@@ -171,7 +171,7 @@ private:
 
 	/** List of currently existing browser windows */
 #if WITH_CEF3
-	TArray<TWeakPtr<FCEFWebBrowserWindow>>	WindowInterfaces;
+	TArray<TWeakPtr<FCEFWebBrowserWindow>> WindowInterfaces;
 #elif PLATFORM_IOS || PLATFORM_SPECIFIC_WEB_BROWSER || (PLATFORM_ANDROID && USE_ANDROID_JNI)
 	TArray<TWeakPtr<IWebBrowserWindow>>	WindowInterfaces;
 #endif
@@ -181,6 +181,10 @@ private:
 
 	TSharedRef<IWebBrowserWindowFactory> WebBrowserWindowFactory;
 
+#if !WITH_CEF3 && (PLATFORM_IOS || (PLATFORM_ANDROID && USE_ANDROID_JNI))
+	FString UserAgentApplication;
+#endif
+
 	bool bDevToolsShortcutEnabled;
 
 	bool bJSBindingsToLoweringEnabled;
