Do not wait for android PSO compiler services to bind.
 - Fixes potential ANR by not requiring response (and blocking android event thread)
 - PSO service usage should check AndroidThunkJava_AreProgramServicesReady first.
 - AndroidThunkJava_HaveServicesFailed indicates binding has failed.

diff --git a/Engine/Build/Android/Java/src/com/epicgames/unreal/psoservices/PSOProgramServiceAccessor.java b/Engine/Build/Android/Java/src/com/epicgames/unreal/psoservices/PSOProgramServiceAccessor.java
--- a/Engine/Build/Android/Java/src/com/epicgames/unreal/psoservices/PSOProgramServiceAccessor.java
+++ b/Engine/Build/Android/Java/src/com/epicgames/unreal/psoservices/PSOProgramServiceAccessor.java
@@ -18,6 +18,7 @@ import android.os.StrictMode;
 import android.os.Trace;
 import android.util.Log;
 import android.os.ParcelFileDescriptor;
+import android.os.SystemClock;
 
 import com.epicgames.unreal.GameActivity;
 import com.epicgames.unreal.Logger;
@@ -221,11 +222,11 @@ public class PSOProgramServiceAccessor
 
 				if(bUseVulkan)
 				{
-					bSuccess = _PSOProgramServiceAccessor.StartAndWaitForVulkanServices(numServices);
+					bSuccess = _PSOProgramServiceAccessor.StartVulkanServices(numServices);
 				}
 				else
 				{
-					bSuccess = _PSOProgramServiceAccessor.StartAndWaitForServices(numServices, bUseRobustEGLContext);
+					bSuccess = _PSOProgramServiceAccessor.StartServices(numServices, bUseRobustEGLContext);
 				}
 			}
 			catch (Exception e)
@@ -267,7 +268,27 @@ public class PSOProgramServiceAccessor
 		}
 	}
 
-	boolean StartAndWaitForServices(int numServices, boolean bUseRobustEGLContext)
+	public static boolean AndroidThunkJava_AreProgramServicesReady()
+	{
+		boolean bSuccess = _PSOProgramServiceAccessor.ServiceInstances.length > 0;
+		for (OGLServiceInstance ServiceInstance : _PSOProgramServiceAccessor.ServiceInstances )
+		{
+			bSuccess = bSuccess && ServiceInstance.IsServiceBound();
+		}
+		return bSuccess;
+	}
+
+	public static boolean AndroidThunkJava_HaveServicesFailed()
+	{
+		boolean bHasFailed = _PSOProgramServiceAccessor.ServiceInstances.length == 0;
+		for (OGLServiceInstance ServiceInstance : _PSOProgramServiceAccessor.ServiceInstances )
+		{
+			bHasFailed = bHasFailed || ServiceInstance.HasBindFailed();
+		}
+		return bHasFailed;
+	}
+
+	boolean StartServices(int numServices, boolean bUseRobustEGLContext)
 	{
 		numServices = Math.max(1, Math.min(numServices, ServiceClassTypes.length));
 
@@ -280,12 +301,12 @@ public class PSOProgramServiceAccessor
 		boolean bSuccess = true;
 		for (OGLServiceInstance ServiceInstance : ServiceInstances )
 		{
-			bSuccess = bSuccess && ServiceInstance.doBindAndWait();
+			bSuccess = bSuccess && ServiceInstance.doBind();
 		}
 		return bSuccess;
 	}
 
-	boolean StartAndWaitForVulkanServices(int numServices)
+	boolean StartVulkanServices(int numServices)
 	{
 		numServices = Math.max(1, Math.min(numServices, VulkanServiceClassTypes.length));
 
@@ -298,7 +319,7 @@ public class PSOProgramServiceAccessor
 		boolean bSuccess = true;
 		for (OGLServiceInstance ServiceInstance : ServiceInstances )
 		{
-			bSuccess = bSuccess && ServiceInstance.doBindAndWait();
+			bSuccess = bSuccess && ServiceInstance.doBind();
 		}
 		return bSuccess;
 	}
@@ -329,23 +350,23 @@ public class PSOProgramServiceAccessor
 	}
 
 	static final private Object ProgramLinkLock = new Object();
-	public static JNIProgramLinkResponse AndroidThunkJava_OGLRemoteProgramLink(byte[] ContextData, String VertexShader, String PixelShader, String ComputeShader)
+	public static JNIProgramLinkResponse AndroidThunkJava_OGLRemoteProgramLink(byte[] ContextData, String VertexShader, String PixelShader, String ComputeShader, boolean bAllowTimeOuts)
 	{
 		if(GameActivity.IsActivityPaused())
 		{
 			// one at a time when backgrounded.
 			synchronized (ProgramLinkLock)
 			{
-				return OGLRemoteProgramLink_internal(ContextData, VertexShader, PixelShader, ComputeShader);
+				return OGLRemoteProgramLink_internal(ContextData, VertexShader, PixelShader, ComputeShader, bAllowTimeOuts);
 			}
 		}
 		else
 		{
-			return OGLRemoteProgramLink_internal(ContextData, VertexShader, PixelShader, ComputeShader);
+			return OGLRemoteProgramLink_internal(ContextData, VertexShader, PixelShader, ComputeShader, bAllowTimeOuts);
 		}
 	}
 
-	private static JNIProgramLinkResponse OGLRemoteProgramLink_internal(byte[] ContextData, String VertexShader, String PixelShader, String ComputeShader)
+	private static JNIProgramLinkResponse OGLRemoteProgramLink_internal(byte[] ContextData, String VertexShader, String PixelShader, String ComputeShader, boolean bAllowTimeOuts)
 	{
 		try
 		{
@@ -427,25 +448,25 @@ public class PSOProgramServiceAccessor
 				{
 					int pendingJobs = thisInstance.PendingJobs.incrementAndGet();
 
-					long tStartTime = System.nanoTime();
-					// wait for 10s before giving up on the compile job.
-					long tTimeoutInMS = 10000;
+					long tStartTimeMS = System.currentTimeMillis();
+					// if we're allowed to time out, wait for 10s before giving up on the compile job.
+					long tTimeoutInMS = bAllowTimeOuts ? 10000 : Long.MAX_VALUE;
 
 					while (pendingResponse.ResponseState != JobResponse.ResponseStateEnum.Responded)
 					{
 						// we must loop as waits can randomly wake.
 						pendingResponse.SyncObj.wait(1000);
 
-						if ((System.nanoTime() - tStartTime) >= (tTimeoutInMS * 1000000))
+						if ((System.currentTimeMillis() - tStartTimeMS) >= tTimeoutInMS)
 						{
-							Log.error("OGLRemoteProgramLink TIMED OUT WAITING " + ThisJobID + " for " + (System.nanoTime() - tStartTime) / 1000000 + "ms. pending tasks "+thisInstance.PendingJobs.get());
+							Log.error("OGLRemoteProgramLink TIMED OUT WAITING " + ThisJobID + " for " + (System.currentTimeMillis() - tStartTimeMS) + "ms. pending tasks "+thisInstance.PendingJobs.get()+"limit "+tTimeoutInMS);
 							//timeout
 							SyncObs.remove(ThisJobID);
 							thisInstance.ReadBackServiceLog();
 							return null;
 						}
 					}
-					long totalWaitTimeMS = (System.nanoTime() - tStartTime) / 1000000;
+					long totalWaitTimeMS = System.currentTimeMillis() - tStartTimeMS;
 					if(totalWaitTimeMS > 2500)
 					{
 						Log.verbose("OGLRemoteProgramLink responded " + ThisJobID + " total wait time "+totalWaitTimeMS+" ms. pending tasks "+thisInstance.PendingJobs.get());
@@ -995,36 +1016,15 @@ public class PSOProgramServiceAccessor
 			}
 			return mShouldUnbind;
 		}
-
-		boolean doBindAndWait()
+		private long bindStartTime = 0;
+		boolean doBind()
 		{
 			synchronized (mConnection.mConnectionSync)
 			{
 				boolean bSuccess = doBindService();
-				long tBefore = System.nanoTime();
-				long tTimeoutInMS = 10000;
-
 				if(bSuccess)
 				{
-					try
-					{
-						while (mBound.get() == -1)
-						{
-							mConnection.mConnectionSync.wait(1000);
-							if ((System.nanoTime() - tBefore) >= (tTimeoutInMS * 1000000))
-							{
-								Log.error("OGLRemoteProgramLink "+Name()+" TIMED OUT waiting for service bind " + (System.nanoTime() - tBefore) / 1000000 + "ms.");
-								//timeout
-								bSuccess = false;
-								break;
-							}
-						}
-					}
-					catch (Exception e)
-					{
-						bSuccess = false;
-						e.printStackTrace();
-					}
+					bindStartTime = SystemClock.uptimeMillis();
 				}
 				return bSuccess;
 			}
@@ -1039,6 +1039,7 @@ public class PSOProgramServiceAccessor
 				mContext.unbindService(mConnection);
 				mShouldUnbind = false;
 				mBound.set(-1);
+				bindStartTime = 0;
 			}
 		}
 
@@ -1047,6 +1048,11 @@ public class PSOProgramServiceAccessor
 			return mBound.get() == 1;
 		}
 
+		boolean HasBindFailed()
+		{
+			return !mShouldUnbind || (!IsServiceBound() && (SystemClock.uptimeMillis()-bindStartTime>10000));
+		}
+
 		boolean SendMessage(Message msg) throws RemoteException
 		{
 			if( mBound.get() == 1)
diff --git a/Engine/Source/Runtime/OpenGLDrv/Private/Android/AndroidOpenGL.cpp b/Engine/Source/Runtime/OpenGLDrv/Private/Android/AndroidOpenGL.cpp
--- a/Engine/Source/Runtime/OpenGLDrv/Private/Android/AndroidOpenGL.cpp
+++ b/Engine/Source/Runtime/OpenGLDrv/Private/Android/AndroidOpenGL.cpp
@@ -92,6 +92,8 @@ struct FOpenGLRemoteGLProgramCompileJNI
 	jmethodID DispatchProgramLink = 0;
 	jmethodID StartRemoteProgramLink = 0;
 	jmethodID StopRemoteProgramLink = 0;
+	jmethodID AreProgramServicesReady = 0;
+	jmethodID HaveServicesFailed = 0;
 	jclass ProgramResponseClass = 0;
 	jfieldID ProgramResponse_SuccessField = 0;
 	jfieldID ProgramResponse_ErrorField = 0;
@@ -115,12 +117,16 @@ struct FOpenGLRemoteGLProgramCompileJNI
 		CHECK_JNI_EXCEPTIONS(Env);
 		if(OGLServiceAccessor)
 		{
-			DispatchProgramLink = FJavaWrapper::FindStaticMethod(Env, OGLServiceAccessor, "AndroidThunkJava_OGLRemoteProgramLink", "([BLjava/lang/String;Ljava/lang/String;Ljava/lang/String;)Lcom/epicgames/unreal/psoservices/PSOProgramServiceAccessor$JNIProgramLinkResponse;", false);
+			DispatchProgramLink = FJavaWrapper::FindStaticMethod(Env, OGLServiceAccessor, "AndroidThunkJava_OGLRemoteProgramLink", "([BLjava/lang/String;Ljava/lang/String;Ljava/lang/String;Z)Lcom/epicgames/unreal/psoservices/PSOProgramServiceAccessor$JNIProgramLinkResponse;", false);
 			CHECK_JNI_EXCEPTIONS(Env);
 			StartRemoteProgramLink = FJavaWrapper::FindStaticMethod(Env, OGLServiceAccessor, "AndroidThunkJava_StartRemoteProgramLink", "(IZZ)Z", false);
 			CHECK_JNI_EXCEPTIONS(Env);
 			StopRemoteProgramLink = FJavaWrapper::FindStaticMethod(Env, OGLServiceAccessor, "AndroidThunkJava_StopRemoteProgramLink", "()V", false);
 			CHECK_JNI_EXCEPTIONS(Env);
+			AreProgramServicesReady = FJavaWrapper::FindStaticMethod(Env, OGLServiceAccessor, "AndroidThunkJava_AreProgramServicesReady", "()Z", false);
+			CHECK_JNI_EXCEPTIONS(Env);
+			HaveServicesFailed = FJavaWrapper::FindStaticMethod(Env, OGLServiceAccessor, "AndroidThunkJava_HaveServicesFailed", "()Z", false);
+			CHECK_JNI_EXCEPTIONS(Env);
 			ProgramResponseClass = AndroidJavaEnv::FindJavaClassGlobalRef("com/epicgames/unreal/psoservices/PSOProgramServiceAccessor$JNIProgramLinkResponse");
 			CHECK_JNI_EXCEPTIONS(Env);
 			ProgramResponse_SuccessField = FJavaWrapper::FindField(Env, ProgramResponseClass, "bCompileSuccess", "Z", true);
@@ -133,7 +139,7 @@ struct FOpenGLRemoteGLProgramCompileJNI
 			CHECK_JNI_EXCEPTIONS(Env);
 		}
 
-		bAllFound = OGLServiceAccessor && DispatchProgramLink && StartRemoteProgramLink && StopRemoteProgramLink && ProgramResponseClass && ProgramResponse_SuccessField && ProgramResponse_CompiledBinaryField && ProgramResponse_ErrorField && ProgramResponse_SHMOutputHandleField;
+		bAllFound = OGLServiceAccessor && DispatchProgramLink && StartRemoteProgramLink && StopRemoteProgramLink && AreProgramServicesReady && HaveServicesFailed && ProgramResponseClass && ProgramResponse_SuccessField && ProgramResponse_CompiledBinaryField && ProgramResponse_ErrorField && ProgramResponse_SHMOutputHandleField;
 		UE_CLOG(!bAllFound, LogRHI, Fatal, TEXT("Failed to find JNI GL remote program compiler."));
 	}
 }OpenGLRemoteGLProgramCompileJNI;
@@ -1096,37 +1102,49 @@ void FAndroidOpenGL::ProcessExtensions(const FString& ExtensionsString)
 	}
 }
 
+static bool GRemoteCompileServicesStarted = false;
 static bool GRemoteCompileServicesActive = false;
 extern bool AreAndroidOpenGLRemoteCompileServicesAvailable();
 
-bool AreAndroidOpenGLRemoteCompileServicesActive()
-{
-	return GRemoteCompileServicesActive && AreAndroidOpenGLRemoteCompileServicesAvailable();
-}
-
 bool FAndroidOpenGL::AreRemoteCompileServicesActive()
 {
-	return AreAndroidOpenGLRemoteCompileServicesActive();
+	if (GRemoteCompileServicesStarted && AreAndroidOpenGLRemoteCompileServicesAvailable())
+	{
+		if (!GRemoteCompileServicesActive)
+		{
+			JNIEnv* Env = FAndroidApplication::GetJavaEnv();
+			GRemoteCompileServicesActive = (bool)Env->CallStaticBooleanMethod(OpenGLRemoteGLProgramCompileJNI.OGLServiceAccessor, OpenGLRemoteGLProgramCompileJNI.AreProgramServicesReady);
+			if (!GRemoteCompileServicesActive)
+			{
+				if ((bool)Env->CallStaticBooleanMethod(OpenGLRemoteGLProgramCompileJNI.OGLServiceAccessor, OpenGLRemoteGLProgramCompileJNI.HaveServicesFailed))
+				{
+					UE_LOG(LogRHI, Error, TEXT("Remote compile services failed to start."));
+					StopRemoteCompileServices();
+				}
+			}
+		}
+		return GRemoteCompileServicesActive;
+	}
+	return false;
 }
 
-bool FAndroidOpenGL::StartAndWaitForRemoteCompileServices(int NumServices)
+bool FAndroidOpenGL::StartRemoteCompileServices(int NumServices)
 {
-	bool bResult = false;
 	JNIEnv* Env = FAndroidApplication::GetJavaEnv();
 
-	if (Env && AreAndroidOpenGLRemoteCompileServicesAvailable())
+	if (Env && AreAndroidOpenGLRemoteCompileServicesAvailable() && !GRemoteCompileServicesStarted)
 	{
 		bool bUseRobustContexts = AndroidEGL::GetInstance()->IsUsingRobustContext();
-		bResult = (bool)Env->CallStaticBooleanMethod(OpenGLRemoteGLProgramCompileJNI.OGLServiceAccessor, OpenGLRemoteGLProgramCompileJNI.StartRemoteProgramLink, (jint)NumServices, (jboolean)bUseRobustContexts, /*bUseVulkan*/(jboolean)false);
-		GRemoteCompileServicesActive = bResult;
+		GRemoteCompileServicesStarted = (bool)Env->CallStaticBooleanMethod(OpenGLRemoteGLProgramCompileJNI.OGLServiceAccessor, OpenGLRemoteGLProgramCompileJNI.StartRemoteProgramLink, (jint)NumServices, (jboolean)bUseRobustContexts, /*bUseVulkan*/(jboolean)false);
 	}
 
-	return bResult;
+	return GRemoteCompileServicesStarted;
 }
 
 void FAndroidOpenGL::StopRemoteCompileServices()
 {
 	GRemoteCompileServicesActive = false;
+	GRemoteCompileServicesStarted = false;
 	JNIEnv* Env = FAndroidApplication::GetJavaEnv();
 
 	if (Env && ensure(AreAndroidOpenGLRemoteCompileServicesAvailable()))
@@ -1155,7 +1173,9 @@ TArray<uint8> FAndroidOpenGL::DispatchAndWaitForRemoteGLProgramCompile(const TAr
 		auto jCS = NewScopedJavaObject(Env, Env->NewStringUTF(TCHAR_TO_UTF8(ANSI_TO_TCHAR(ComputeGlslCode.IsEmpty() ? "" : ComputeGlslCode.GetData()))));
 		auto ProgramKeyBuffer = NewScopedJavaObject(Env, Env->NewByteArray(ContextData.Num()));
 		Env->SetByteArrayRegion(*ProgramKeyBuffer, 0, ContextData.Num(), reinterpret_cast<const jbyte*>(ContextData.GetData()));
-		auto ProgramResponseObj = NewScopedJavaObject(Env, Env->CallStaticObjectMethod(OpenGLRemoteGLProgramCompileJNI.OGLServiceAccessor, OpenGLRemoteGLProgramCompileJNI.DispatchProgramLink, *ProgramKeyBuffer, *jVS, *jPS, *jCS));
+		// dont time out if the debugger is attached.
+		bool bEnableTimeOuts = !FPlatformMisc::IsDebuggerPresent();
+		auto ProgramResponseObj = NewScopedJavaObject(Env, Env->CallStaticObjectMethod(OpenGLRemoteGLProgramCompileJNI.OGLServiceAccessor, OpenGLRemoteGLProgramCompileJNI.DispatchProgramLink, *ProgramKeyBuffer, *jVS, *jPS, *jCS, bEnableTimeOuts));
  		CHECK_JNI_EXCEPTIONS(Env);
 
 		if(*ProgramResponseObj)
diff --git a/Engine/Source/Runtime/OpenGLDrv/Private/Android/AndroidOpenGL.h b/Engine/Source/Runtime/OpenGLDrv/Private/Android/AndroidOpenGL.h
--- a/Engine/Source/Runtime/OpenGLDrv/Private/Android/AndroidOpenGL.h
+++ b/Engine/Source/Runtime/OpenGLDrv/Private/Android/AndroidOpenGL.h
@@ -142,7 +142,7 @@ struct FAndroidOpenGL : public FOpenGLES
 	/* interface to remote GLES program compiler */
 	static TArray<uint8> DispatchAndWaitForRemoteGLProgramCompile(const TArrayView<uint8> ContextData, const TArray<ANSICHAR>& VertexGlslCode, const TArray<ANSICHAR>& PixelGlslCode, const TArray<ANSICHAR>& ComputeGlslCode, FString& FailureMessageOUT);
 	static bool AreRemoteCompileServicesActive();
-	static bool StartAndWaitForRemoteCompileServices(int NumServices);
+	static bool StartRemoteCompileServices(int NumServices);
 	static void StopRemoteCompileServices();
 };
 
diff --git a/Engine/Source/Runtime/OpenGLDrv/Private/OpenGLProgramBinaryFileCache.cpp b/Engine/Source/Runtime/OpenGLDrv/Private/OpenGLProgramBinaryFileCache.cpp
--- a/Engine/Source/Runtime/OpenGLDrv/Private/OpenGLProgramBinaryFileCache.cpp
+++ b/Engine/Source/Runtime/OpenGLDrv/Private/OpenGLProgramBinaryFileCache.cpp
@@ -554,7 +554,7 @@ void FOpenGLProgramBinaryCache::OnShaderPipelineCacheOpened(FString const& Name,
 			check(!FAndroidOpenGL::AreRemoteCompileServicesActive() || IsPrecachingEnabled());
 			if(FAndroidOpenGL::AreRemoteCompileServicesActive() == false)
 			{
-				FAndroidOpenGL::StartAndWaitForRemoteCompileServices(GNumRemoteProgramCompileServices);
+				FAndroidOpenGL::StartRemoteCompileServices(GNumRemoteProgramCompileServices);
 			}
 		}
 #endif
diff --git a/Engine/Source/Runtime/VulkanRHI/Private/Android/VulkanAndroidPlatform.cpp b/Engine/Source/Runtime/VulkanRHI/Private/Android/VulkanAndroidPlatform.cpp
--- a/Engine/Source/Runtime/VulkanRHI/Private/Android/VulkanAndroidPlatform.cpp
+++ b/Engine/Source/Runtime/VulkanRHI/Private/Android/VulkanAndroidPlatform.cpp
@@ -1348,6 +1348,7 @@ void PipelineToBinary(FVulkanDevice* Device, const VkGraphicsPipelineCreateInfo*
 #define CHECK_JNI_EXCEPTIONS(env)  if (env->ExceptionCheck()) {env->ExceptionDescribe();env->ExceptionClear();}
 #endif
 
+static bool GRemoteCompileServicesStarted = false;
 static bool GRemoteCompileServicesActive = false;
 
 struct FVKRemoteProgramCompileJNI
@@ -1356,6 +1357,8 @@ struct FVKRemoteProgramCompileJNI
 	jmethodID DispatchPSOCompile = 0;
 	jmethodID DispatchPSOCompileShm = 0;
 	jmethodID StartRemoteProgramLink = 0;
+	jmethodID HaveServicesFailed = 0;
+	jmethodID AreProgramServicesReady = 0;	
 	jmethodID StopRemoteProgramLink = 0;
 	jclass ProgramResponseClass = 0;
 	jfieldID ProgramResponse_SuccessField = 0;
@@ -1390,6 +1393,10 @@ struct FVKRemoteProgramCompileJNI
 			CHECK_JNI_EXCEPTIONS(Env);
 			StartRemoteProgramLink = FJavaWrapper::FindStaticMethod(Env, PSOServiceAccessor, "AndroidThunkJava_StartRemoteProgramLink", "(IZZ)Z", false);
 			CHECK_JNI_EXCEPTIONS(Env);
+			HaveServicesFailed = FJavaWrapper::FindStaticMethod(Env, PSOServiceAccessor, "AndroidThunkJava_HaveServicesFailed", "()Z", false);
+			CHECK_JNI_EXCEPTIONS(Env);
+			AreProgramServicesReady = FJavaWrapper::FindStaticMethod(Env, PSOServiceAccessor, "AndroidThunkJava_AreProgramServicesReady", "()Z", false);
+			CHECK_JNI_EXCEPTIONS(Env);
 			StopRemoteProgramLink = FJavaWrapper::FindStaticMethod(Env, PSOServiceAccessor, "AndroidThunkJava_StopRemoteProgramLink", "()V", false);
 			CHECK_JNI_EXCEPTIONS(Env);
 			ProgramResponseClass = AndroidJavaEnv::FindJavaClassGlobalRef("com/epicgames/unreal/psoservices/PSOProgramServiceAccessor$JNIProgramLinkResponse");
@@ -1404,7 +1411,7 @@ struct FVKRemoteProgramCompileJNI
 			CHECK_JNI_EXCEPTIONS(Env);
 		}
 
-		bAllFound = PSOServiceAccessor && DispatchPSOCompile && DispatchPSOCompileShm && StartRemoteProgramLink && StopRemoteProgramLink && ProgramResponseClass && ProgramResponse_SuccessField && ProgramResponse_CompiledBinaryField && ProgramResponse_ErrorField && ProgramResponse_SHMOutputHandleField;
+		bAllFound = PSOServiceAccessor && DispatchPSOCompile && DispatchPSOCompileShm && StartRemoteProgramLink && HaveServicesFailed && AreProgramServicesReady && StopRemoteProgramLink && ProgramResponseClass && ProgramResponse_SuccessField && ProgramResponse_CompiledBinaryField && ProgramResponse_ErrorField && ProgramResponse_SHMOutputHandleField;
 		UE_CLOG(!bAllFound, LogRHI, Fatal, TEXT("Failed to find JNI Vulkan remote program compiler."));
 	}
 }VKRemoteProgramCompileJNI;
@@ -1427,34 +1434,45 @@ static bool AreAndroidVulkanRemoteCompileServicesAvailable()
 	return RemoteCompileService;
 }
 
-bool AreAndroidVulkanRemoteCompileServicesActive()
-{
-	return GRemoteCompileServicesActive && AreAndroidVulkanRemoteCompileServicesAvailable();
-}
-
 bool FVulkanAndroidPlatform::AreRemoteCompileServicesActive()
 {
-	return AreAndroidVulkanRemoteCompileServicesActive();
+	if (GRemoteCompileServicesStarted && AreAndroidVulkanRemoteCompileServicesAvailable())
+	{
+		if (!GRemoteCompileServicesActive)
+		{
+			JNIEnv* Env = FAndroidApplication::GetJavaEnv();
+			GRemoteCompileServicesActive = (bool)Env->CallStaticBooleanMethod(VKRemoteProgramCompileJNI.PSOServiceAccessor, VKRemoteProgramCompileJNI.AreProgramServicesReady);
+			if (!GRemoteCompileServicesActive)
+			{
+				if ((bool)Env->CallStaticBooleanMethod(VKRemoteProgramCompileJNI.PSOServiceAccessor, VKRemoteProgramCompileJNI.HaveServicesFailed))
+				{
+					UE_LOG(LogRHI, Error, TEXT("Remote compile services failed to start."));
+					StopRemoteCompileServices();
+				}
+			}
+		}
+		return GRemoteCompileServicesActive;
+	}
+	return false;
 }
 
-bool FVulkanAndroidPlatform::StartAndWaitForRemoteCompileServices(int NumServices)
+bool FVulkanAndroidPlatform::StartRemoteCompileServices(int NumServices)
 {
-	bool bResult = false;
 	JNIEnv* Env = FAndroidApplication::GetJavaEnv();
 	
 	VKRemoteProgramCompileJNI.Init(Env);
 
-	if (Env && AreAndroidVulkanRemoteCompileServicesAvailable())
+	if (Env && AreAndroidVulkanRemoteCompileServicesAvailable() && !GRemoteCompileServicesStarted)
 	{
-		bResult = (bool)Env->CallStaticBooleanMethod(VKRemoteProgramCompileJNI.PSOServiceAccessor, VKRemoteProgramCompileJNI.StartRemoteProgramLink, (jint)NumServices, /*bUseRobustEGLContext*/(jboolean)false, /*bUseVulkan*/(jboolean)true);
-		GRemoteCompileServicesActive = bResult;
+		GRemoteCompileServicesStarted = (bool)Env->CallStaticBooleanMethod(VKRemoteProgramCompileJNI.PSOServiceAccessor, VKRemoteProgramCompileJNI.StartRemoteProgramLink, (jint)NumServices, /*bUseRobustEGLContext*/(jboolean)false, /*bUseVulkan*/(jboolean)true);
 	}
 
-	return bResult;
+	return GRemoteCompileServicesStarted;
 }
 
 void FVulkanAndroidPlatform::StopRemoteCompileServices()
 {
+	GRemoteCompileServicesStarted = false;
 	GRemoteCompileServicesActive = false;
 	JNIEnv* Env = FAndroidApplication::GetJavaEnv();
 
@@ -1473,7 +1491,7 @@ VkPipelineCache FVulkanAndroidPlatform::PrecompilePSO(FVulkanDevice* Device, con
 {
 	FString FailureMessageOUT;
 	
-	if (!AreAndroidVulkanRemoteCompileServicesActive())
+	if (!AreRemoteCompileServicesActive())
 	{
 		return VK_NULL_HANDLE;
 	}
@@ -1642,7 +1660,7 @@ void FVulkanAndroidPlatform::PostInitGPU(const FVulkanDevice& InDevice)
 	static const auto CVarVulkanPSOPrecaching = IConsoleManager::Get().FindConsoleVariable(TEXT("r.Vulkan.AllowPSOPrecaching"));	
 	if (CVarNumRemoteProgramCompileServices->GetInt() && CVarChunkedPSOCache->GetInt() && CVarPSOPrecaching->GetInt() && CVarVulkanPSOPrecaching->GetInt())
 	{
-		FVulkanAndroidPlatform::StartAndWaitForRemoteCompileServices(CVarNumRemoteProgramCompileServices->GetInt());
+		FVulkanAndroidPlatform::StartRemoteCompileServices(CVarNumRemoteProgramCompileServices->GetInt());
 	}
 }
 
diff --git a/Engine/Source/Runtime/VulkanRHI/Private/Android/VulkanAndroidPlatform.h b/Engine/Source/Runtime/VulkanRHI/Private/Android/VulkanAndroidPlatform.h
--- a/Engine/Source/Runtime/VulkanRHI/Private/Android/VulkanAndroidPlatform.h
+++ b/Engine/Source/Runtime/VulkanRHI/Private/Android/VulkanAndroidPlatform.h
@@ -149,7 +149,7 @@ public:
 	static VkPipelineCache PrecompilePSO(FVulkanDevice* Device, const TArrayView<uint8> OptionalPSOCacheData, const VkGraphicsPipelineCreateInfo* PipelineInfo, FGfxPipelineDesc* GfxEntry, const FVulkanRenderTargetLayout* RTLayout, TArrayView<uint32_t> VS, TArrayView<uint32_t> PS, size_t& AfterSize);
 
 	static bool AreRemoteCompileServicesActive();
-	static bool StartAndWaitForRemoteCompileServices(int NumServices);
+	static bool StartRemoteCompileServices(int NumServices);
 	static void StopRemoteCompileServices();
 
 	// Do not attempt to immediately recreate swapchain
diff --git a/Engine/Source/Runtime/VulkanRHI/Private/VulkanPipeline.cpp b/Engine/Source/Runtime/VulkanRHI/Private/VulkanPipeline.cpp
--- a/Engine/Source/Runtime/VulkanRHI/Private/VulkanPipeline.cpp
+++ b/Engine/Source/Runtime/VulkanRHI/Private/VulkanPipeline.cpp
@@ -564,7 +564,7 @@ void FVulkanPipelineStateCacheManager::OnShaderPipelineCacheOpened(FString const
 #if PLATFORM_ANDROID
 		if (GNumRemoteProgramCompileServices)
 		{
-			FVulkanAndroidPlatform::StartAndWaitForRemoteCompileServices(GNumRemoteProgramCompileServices);
+			FVulkanAndroidPlatform::StartRemoteCompileServices(GNumRemoteProgramCompileServices);
 		}
 #endif
 	}
