Fix for exit crash in PSO Service

diff --git a/Engine/Build/Android/Prebuilt/PSOService/Source/PSOService.cpp b/Engine/Build/Android/Prebuilt/PSOService/Source/PSOService.cpp
--- a/Engine/Build/Android/Prebuilt/PSOService/Source/PSOService.cpp
+++ b/Engine/Build/Android/Prebuilt/PSOService/Source/PSOService.cpp
@@ -235,6 +235,8 @@ public:
 		DestroyPipelineCache();
 		vkDestroyDevice(Device, nullptr);
 		vkDestroyInstance(Instance, nullptr);
+		Device = VK_NULL_HANDLE;
+		Instance = VK_NULL_HANDLE;
 	}
 
 	struct GraphicsPipelineCreateInfo
