Fix missing symbol issue with GetSharedPreference functions

#android

diff --git a/Engine/Build/Android/Java/src/com/epicgames/unreal/GameActivity.java.template b/Engine/Build/Android/Java/src/com/epicgames/unreal/GameActivity.java.template
--- a/Engine/Build/Android/Java/src/com/epicgames/unreal/GameActivity.java.template
+++ b/Engine/Build/Android/Java/src/com/epicgames/unreal/GameActivity.java.template
@@ -4788,7 +4788,7 @@ public class GameActivity extends $${gameActivitySuperClass}$$ implements Surfac
 		Context context = getApplicationContext();
 		SharedPreferences preferences = context.getSharedPreferences(group, context.MODE_PRIVATE);
 		SharedPreferences.Editor editor = preferences.edit();
-		editor.putInt(key, key, value);
+		editor.putInt(key, value);
 		editor.commit();
 	}
 
@@ -4808,7 +4808,7 @@ public class GameActivity extends $${gameActivitySuperClass}$$ implements Surfac
 		editor.commit();
 	}
 
-	public String AndroidThunkJava_GetSharedPreferenceString(String group, String key, string defaultValue)
+	public String AndroidThunkJava_GetSharedPreferenceString(String group, String key, String defaultValue)
 	{
 		Context context = getApplicationContext();
 		SharedPreferences preferences = context.getSharedPreferences(group, context.MODE_PRIVATE);
