GameActivity.java.template: Update deviceScreenOrientation onConfigurationChanged to make sure we can detect the real/allowed device orientation.
AndroidApplication.cpp: Update widgets safezone when the window size changes (i.e the device orientation changed)

#rnx

diff --git a/Engine/Build/Android/Java/src/com/epicgames/unreal/GameActivity.java.template b/Engine/Build/Android/Java/src/com/epicgames/unreal/GameActivity.java.template
--- a/Engine/Build/Android/Java/src/com/epicgames/unreal/GameActivity.java.template
+++ b/Engine/Build/Android/Java/src/com/epicgames/unreal/GameActivity.java.template
@@ -4071,7 +4071,9 @@ public class GameActivity extends $${gameActivitySuperClass}$$ implements Surfac
 		super.onConfigurationChanged(newConfig);
 //$${gameActivityonConfigurationChangedAdditions}$$
 		// forward the orientation
-		nativeOnOrientationChanged(newConfig.orientation);
+		Log.debug("==============> onConfigurationChanged: orientation : " + newConfig.orientation);
+		deviceScreenOrientation = newConfig.orientation;
+		nativeOnOrientationChanged(deviceScreenOrientation);
 	}
 
 	public void surfaceChanged(SurfaceHolder holder, int format, int width, int height)
diff --git a/Engine/Source/Runtime/ApplicationCore/Private/Android/AndroidApplication.cpp b/Engine/Source/Runtime/ApplicationCore/Private/Android/AndroidApplication.cpp
--- a/Engine/Source/Runtime/ApplicationCore/Private/Android/AndroidApplication.cpp
+++ b/Engine/Source/Runtime/ApplicationCore/Private/Android/AndroidApplication.cpp
@@ -100,6 +100,7 @@ void FAndroidApplication::PollGameDeviceState( const float TimeDelta )
 		FDisplayMetrics DisplayMetrics;
 		FDisplayMetrics::RebuildDisplayMetrics(DisplayMetrics);
 		BroadcastDisplayMetricsChanged(DisplayMetrics);
+		FCoreDelegates::OnSafeFrameChangedEvent.Broadcast();
 
 		// the cursor needs to compute the proper slate scaling factor each time the display metrics change
 		TSharedPtr<FAndroidCursor> AndroidCursor = StaticCastSharedPtr<FAndroidCursor>(Cursor);
