Fixed updating device orientation in Android's onConfigurationChanged callback
It was using the value in the Configuration object which is not directly convertible to EDeviceScreenOrientation (which is what is expected for deviceScreenOrientation)

diff --git a/Engine/Build/Android/Java/src/com/epicgames/unreal/GameActivity.java.template b/Engine/Build/Android/Java/src/com/epicgames/unreal/GameActivity.java.template
--- a/Engine/Build/Android/Java/src/com/epicgames/unreal/GameActivity.java.template
+++ b/Engine/Build/Android/Java/src/com/epicgames/unreal/GameActivity.java.template
@@ -4070,10 +4070,9 @@ public class GameActivity extends $${gameActivitySuperClass}$$ implements Surfac
 	{
 		super.onConfigurationChanged(newConfig);
 //$${gameActivityonConfigurationChangedAdditions}$$
-		// forward the orientation
-		Log.debug("==============> onConfigurationChanged: orientation : " + newConfig.orientation);
-		deviceScreenOrientation = newConfig.orientation;
-		nativeOnOrientationChanged(deviceScreenOrientation);
+		// recalculate current orientation (the one in the Configuration object is only portrait/landscape and does not convert to EDeviceScreenOrientation)
+		// AndroidThunkJava_GetDeviceOrientation will call nativeOnOrientationChanged if the current orientation actually changes from last time
+		AndroidThunkJava_GetDeviceOrientation();
 	}
 
 	public void surfaceChanged(SurfaceHolder holder, int format, int width, int height)
@@ -6987,7 +6986,7 @@ public class GameActivity extends $${gameActivitySuperClass}$$ implements Surfac
 		if (deviceScreenOrientation != currentScreenOrientation)
 		{
 			deviceScreenOrientation = currentScreenOrientation;
-			Log.debug("==============> onSensorChanged: orientation : " + deviceScreenOrientation);
+			Log.debug("==============> GetDeviceOrientation new orientation : " + deviceScreenOrientation);
 			nativeOnOrientationChanged(deviceScreenOrientation);
 		}
 
