[Engine] Add PSO compilation metric

#rnx

diff --git a/Engine/Build/Android/Java/src/com/epicgames/unreal/psoservices/PSOProgramService.java b/Engine/Build/Android/Java/src/com/epicgames/unreal/psoservices/PSOProgramService.java
--- a/Engine/Build/Android/Java/src/com/epicgames/unreal/psoservices/PSOProgramService.java
+++ b/Engine/Build/Android/Java/src/com/epicgames/unreal/psoservices/PSOProgramService.java
@@ -63,6 +63,7 @@ public class PSOProgramService extends Service implements Logger.ILoggerCallback
 	static final String ServiceID_Key = "sid";
 	static final String CompiledProgram_Key = "cpg";
 	static final String SHMem_Key = "shm";
+	static final String CompilationDuration_Key = "cd";
 	static final String JobFail = "f";
 	static final String RobustContextKey = "rbc";
 	public static final String LogDir = "/oglservice/";
@@ -394,9 +395,10 @@ public class PSOProgramService extends Service implements Logger.ILoggerCallback
 	public void ProcessVulkanProgramRequest(Messenger replyTo, int JobID, int ServiceID, byte[] JobContext, byte[] VS, byte[] PS, byte[] PSOData, byte[] PSOCacheData)
 	{
 		ByteBuffer Result = null;
+		float[] PSOCompilationDuration = new float[1];
 		try
 		{
-			byte[] PipelineCache = compileVulkanPSO(JobID, VS, PS, PSOData, PSOCacheData);
+			byte[] PipelineCache = compileVulkanPSO(JobID, VS, PS, PSOData, PSOCacheData, PSOCompilationDuration);
 			Result = ByteBuffer.allocate(PipelineCache.length);
 			Result.order(ByteOrder.nativeOrder());
 			Result.put(PipelineCache);
@@ -411,17 +413,18 @@ public class PSOProgramService extends Service implements Logger.ILoggerCallback
 		if(Result != null)
 		{
 			//logger.verbose(JobID+" SendSuccess()  ");
-			SendSuccess(replyTo, ServiceID, JobID, JobContext, Result.array());
+			SendSuccess(replyTo, ServiceID, JobID, JobContext, Result.array(), PSOCompilationDuration[0]);
 		}
 	}
 
 	public void ProcessVulkanProgramRequestSHM(Messenger replyTo, int JobID, int ServiceID, byte[] JobContext, int SHMemFD, long VSSize, long PSSize, long PSODataSize, long PSOCacheDataSize)
 	{
 		int SHMResultFD = -1;
+		float[] PSOCompilationDuration = new float[1];
 		try
 		{
 			// native code returns a shared FD which must be adopted in java land.
-			SHMResultFD = compileVulkanPSOSHM(JobID, SHMemFD, VSSize, PSSize, PSODataSize, PSOCacheDataSize);
+			SHMResultFD = compileVulkanPSOSHM(JobID, SHMemFD, VSSize, PSSize, PSODataSize, PSOCacheDataSize, PSOCompilationDuration);
 		}
 		catch (Exception e)
 		{
@@ -432,7 +435,7 @@ public class PSOProgramService extends Service implements Logger.ILoggerCallback
 		{
 			//logger.verbose(JobID+" SendSuccess()  ");
 			// SendSuccess will adopt the shared FD.
-			SendSuccess(replyTo, ServiceID, JobID, JobContext, SHMResultFD );
+			SendSuccess(replyTo, ServiceID, JobID, JobContext, SHMResultFD, PSOCompilationDuration[0]);
 		}
 		else
 		{
@@ -445,8 +448,10 @@ public class PSOProgramService extends Service implements Logger.ILoggerCallback
 	public void ProcessGLProgramRequest(Messenger replyTo, int JobID, int ServiceID, byte[] JobContext, String VS, String PS, String CS)
 	{
 		ByteBuffer Result = null;
+		float CompilationDuration = -1.0f;
 		try
 		{
+			long CompilationStart = System.currentTimeMillis();	
 			if( bEnableTestFailures && JobID == 500)
 			{
 				int splitpoint = VS.length()/2;
@@ -457,6 +462,8 @@ public class PSOProgramService extends Service implements Logger.ILoggerCallback
 			Result = CompileAndLink(JobID, VS, PS, CS);
 			//logger.verbose(JobID+" CompileAndLink()  "+ (Result == null ? "null res " : Result.toString()));
 
+			CompilationDuration = (float)(((double)System.currentTimeMillis() - (double)CompilationStart) / 1000.0f);
+
 		} catch (Exception e)
 		{
 			SendFail(replyTo, ServiceID, JobID, JobContext, e.getMessage());
@@ -465,7 +472,7 @@ public class PSOProgramService extends Service implements Logger.ILoggerCallback
 		if(Result != null)
 		{
 			//logger.verbose(JobID+" SendSuccess()  ");
-			SendSuccess(replyTo, ServiceID, JobID, JobContext, Result.array());
+			SendSuccess(replyTo, ServiceID, JobID, JobContext, Result.array(), CompilationDuration);
 		}
 	}
 
@@ -501,17 +508,17 @@ public class PSOProgramService extends Service implements Logger.ILoggerCallback
 		super.onDestroy();
 	}
 
-	public void SendSuccess(Messenger replyTo, int ServiceID, int JobID, byte[] context, byte[] compiledbinary)
+	public void SendSuccess(Messenger replyTo, int ServiceID, int JobID, byte[] context, byte[] compiledbinary, float compilationDuration)
 	{
-		SendSuccess(replyTo, ServiceID, JobID, context, compiledbinary, -1);
+		SendSuccess(replyTo, ServiceID, JobID, context, compiledbinary, -1, compilationDuration);
 	}
 
-	public void SendSuccess(Messenger replyTo, int ServiceID, int JobID, byte[] context, int ShmMemCompiledFD)
+	public void SendSuccess(Messenger replyTo, int ServiceID, int JobID, byte[] context, int ShmMemCompiledFD, float compilationDuration)
 	{
-		SendSuccess(replyTo, ServiceID, JobID, context, null, ShmMemCompiledFD);
+		SendSuccess(replyTo, ServiceID, JobID, context, null, ShmMemCompiledFD, compilationDuration);
 	}
 
-	public void SendSuccess(Messenger replyTo, int ServiceID, int JobID, byte[] context, byte[] compiledbinary, int ShmMemCompiledFD)
+	public void SendSuccess(Messenger replyTo, int ServiceID, int JobID, byte[] context, byte[] compiledbinary, int ShmMemCompiledFD, float compilationDuration)
 	{
 		if( bEnableTestFailures && JobID == 1000 )
 		{
@@ -546,6 +553,7 @@ public class PSOProgramService extends Service implements Logger.ILoggerCallback
 			}
 			params.putInt(JobID_Key, JobID);
 			params.putInt(ServiceID_Key, ServiceID);
+			params.putFloat(CompilationDuration_Key, compilationDuration);
 			msg.setData(params);
 
 			replyTo.send(msg);
@@ -590,8 +598,8 @@ public class PSOProgramService extends Service implements Logger.ILoggerCallback
 	////////////////////////////////////////////// Vulkan context stuff
 	public native void InitVKDevice();
 	public native void ShutdownVKDevice();
-	public native byte[] CompileVKGFXPSO(byte[] VertexShaderSource, byte[]  PixelShaderSource, byte[] PSOData, byte[] PSOCacheData);
-	public native int CompileVKGFXPSOSHM(int SHMemFD, long VSSize, long PSSize, long PSODataSize, long PSOCacheDataSize);
+	public native byte[] CompileVKGFXPSO(byte[] VertexShaderSource, byte[]  PixelShaderSource, byte[] PSOData, byte[] PSOCacheData, float[] PSOCompilationDuration);
+	public native int CompileVKGFXPSOSHM(int SHMemFD, long VSSize, long PSSize, long PSODataSize, long PSOCacheDataSize, float[] PSOCompilationDuration);
 
 	private void initVulkanContext()
 	{
@@ -600,14 +608,14 @@ public class PSOProgramService extends Service implements Logger.ILoggerCallback
 		endTrace();
 	}
 
-	private byte[] compileVulkanPSO(int JobID, byte [] VertexShaderSource, byte []  PixelShaderSource, byte [] PSOData, byte [] PSOCacheData)
+	private byte[] compileVulkanPSO(int JobID, byte [] VertexShaderSource, byte []  PixelShaderSource, byte [] PSOData, byte [] PSOCacheData, float[] PSOCompilationDuration)
 	{
-		return CompileVKGFXPSO(VertexShaderSource, PixelShaderSource, PSOData, PSOCacheData);
+		return CompileVKGFXPSO(VertexShaderSource, PixelShaderSource, PSOData, PSOCacheData, PSOCompilationDuration);
 	}
 
-	private int compileVulkanPSOSHM(int JobID, int SHMemFD, long VSSize, long PSSize, long PSODataSize, long PSOCacheDataSize)
+	private int compileVulkanPSOSHM(int JobID, int SHMemFD, long VSSize, long PSSize, long PSODataSize, long PSOCacheDataSize, float[] PSOCompilationDuration)
 	{
-		return CompileVKGFXPSOSHM(SHMemFD, VSSize, PSSize, PSODataSize, PSOCacheDataSize);
+		return CompileVKGFXPSOSHM(SHMemFD, VSSize, PSSize, PSODataSize, PSOCacheDataSize, PSOCompilationDuration);
 	}
 
 	private void shutdownVulkanContext()
diff --git a/Engine/Build/Android/Java/src/com/epicgames/unreal/psoservices/PSOProgramServiceAccessor.java b/Engine/Build/Android/Java/src/com/epicgames/unreal/psoservices/PSOProgramServiceAccessor.java
--- a/Engine/Build/Android/Java/src/com/epicgames/unreal/psoservices/PSOProgramServiceAccessor.java
+++ b/Engine/Build/Android/Java/src/com/epicgames/unreal/psoservices/PSOProgramServiceAccessor.java
@@ -120,6 +120,8 @@ public class PSOProgramServiceAccessor
 					byte[] JobContext = msg.getData().getByteArray(PSOProgramService.JobContext_Key);
 					int JobID = msg.getData().getInt(PSOProgramService.JobID_Key, -1);
 					int ServiceID = msg.getData().getInt(PSOProgramService.ServiceID_Key, -1);
+					// not needed here, since we'll just respond with the same data (Bundle)
+					//float CompilationDuration = msg.getData().getFloat(PSOProgramService.CompilationDuration_Key, -1.0f);
 
 					try
 					{
@@ -361,6 +363,7 @@ public class PSOProgramServiceAccessor
 		String ErrorMessage;
 		byte[] CompiledProgram;
 		int SHMOutputHandle;
+		float CompilationDuration;
 	}
 
 	static final private Object ProgramLinkLock = new Object();
@@ -528,6 +531,7 @@ public class PSOProgramServiceAccessor
 			byte[] EngineJobContext = pendingResponse.data.getByteArray(PSOProgramService.JobContext_Key);
 			byte[] CompiledBinary = pendingResponse.data.getByteArray(PSOProgramService.CompiledProgram_Key);
 			int JobID = pendingResponse.data.getInt(PSOProgramService.JobID_Key);
+			float CompilationDuration = pendingResponse.data.getFloat(PSOProgramService.CompilationDuration_Key);
 
 			//Log.verbose("OGLRemoteProgramLink handoff ("+JobID+"), ("+(EngineJobContext == null ? 0 : EngineJobContext.length)+", "+(CompiledBinary==null?0:CompiledBinary.length)+")");
 
@@ -536,6 +540,7 @@ public class PSOProgramServiceAccessor
 			jniResponse.bCompileSuccess = fail == null || fail.isEmpty();
 			jniResponse.ErrorMessage = fail;
 			jniResponse.CompiledProgram = CompiledBinary;
+			jniResponse.CompilationDuration = CompilationDuration;
 
 			return jniResponse;
 		}
@@ -746,6 +751,7 @@ public class PSOProgramServiceAccessor
 			byte[] EngineJobContext = pendingResponse.data.getByteArray(PSOProgramService.JobContext_Key);
 			byte[] CompiledBinary = pendingResponse.data.getByteArray(PSOProgramService.CompiledProgram_Key);
 			int JobID = pendingResponse.data.getInt(PSOProgramService.JobID_Key);
+			float CompilationDuration = pendingResponse.data.getFloat(PSOProgramService.CompilationDuration_Key);
 
 			//Log.verbose("OGLRemoteProgramLink handoff ("+JobID+"), ("+(EngineJobContext == null ? 0 : EngineJobContext.length)+", "+(CompiledBinary==null?0:CompiledBinary.length)+")");
 
@@ -754,6 +760,7 @@ public class PSOProgramServiceAccessor
 			jniResponse.bCompileSuccess = fail == null || fail.isEmpty();
 			jniResponse.ErrorMessage = fail;
 			jniResponse.CompiledProgram = CompiledBinary;
+			jniResponse.CompilationDuration = CompilationDuration;
 
 			return jniResponse;
 		}
@@ -915,6 +922,8 @@ public class PSOProgramServiceAccessor
 
 			int JobID = pendingResponse.data.getInt(PSOProgramService.JobID_Key);
 
+			float CompilationDuration = pendingResponse.data.getFloat(PSOProgramService.CompilationDuration_Key);
+
 			//Log.verbose("OGLRemoteProgramLink handoff ("+JobID+"), ("+(EngineJobContext == null ? 0 : EngineJobContext.length)+")");
 
 			JNIProgramLinkResponse jniResponse = new JNIProgramLinkResponse();
@@ -922,6 +931,7 @@ public class PSOProgramServiceAccessor
 			jniResponse.bCompileSuccess = fail == null || fail.isEmpty();
 			jniResponse.ErrorMessage = fail;
 			jniResponse.SHMOutputHandle = CompiledBinarySharedFD;
+			jniResponse.CompilationDuration = CompilationDuration;
 
 			return jniResponse;
 		}
diff --git a/Engine/Build/Android/Prebuilt/PSOService/Source/PSOService.cpp b/Engine/Build/Android/Prebuilt/PSOService/Source/PSOService.cpp
--- a/Engine/Build/Android/Prebuilt/PSOService/Source/PSOService.cpp
+++ b/Engine/Build/Android/Prebuilt/PSOService/Source/PSOService.cpp
@@ -56,6 +56,14 @@ VKAPI_ATTR VkBool32 VKAPI_CALL VKValidationCallback(
 	return VK_FALSE;
 }
 
+// gets current time in seconds
+double now_s() 
+{
+	struct timespec res;
+	clock_gettime(CLOCK_REALTIME, &res);
+	return (double) res.tv_sec + (double) res.tv_nsec / 1e9;
+}
+
 class FVulkanPSOCompiler
 {
 	bool bInitialized = false;
@@ -989,10 +997,12 @@ void ExitTest()
 	}
 }
 
-JNI_METHOD jobject Java_com_epicgames_unreal_psoservices_PSOProgramService_CompileVKGFXPSO(JNIEnv* jenv, jobject thiz, jbyteArray jVS, jbyteArray jPS, jbyteArray jPSO, jbyteArray jPSOCacheDataSource)
+JNI_METHOD jobject Java_com_epicgames_unreal_psoservices_PSOProgramService_CompileVKGFXPSO(JNIEnv* jenv, jobject thiz, jbyteArray jVS, jbyteArray jPS, jbyteArray jPSO, jbyteArray jPSOCacheDataSource, jfloatArray jCompilationDuration)
 {
 	ExitTest();
 
+	double CompilationStartTime = now_s();
+
 	const uint8_t* VS = (const uint8_t*)jenv->GetByteArrayElements(jVS, nullptr);
 	uint64_t VSSize = jenv->GetArrayLength(jVS);
 	const uint8_t* PS = (const uint8_t*)jenv->GetByteArrayElements(jPS, nullptr);
@@ -1018,14 +1028,26 @@ JNI_METHOD jobject Java_com_epicgames_unreal_psoservices_PSOProgramService_Compi
 		free(BinaryData);
 	}
 
+	double CompilationDuration = now_s() - CompilationStartTime;
+
+	float *CDA = jenv->GetFloatArrayElements(jCompilationDuration, nullptr);
+	if (CDA != nullptr)
+	{
+		CDA[0] = (float)CompilationDuration;
+		jenv->ReleaseFloatArrayElements(jCompilationDuration, CDA, 0);
+	}
+
 	return Data;
 }
 
 // the shared mem version takes an FD and a bunch of offsets.
 // another shared FD containing the result is returned.
-JNI_METHOD jint Java_com_epicgames_unreal_psoservices_PSOProgramService_CompileVKGFXPSOSHM(JNIEnv* jenv, jobject thiz, jint SHMemFD, jlong jVSSize, jlong jPSSize, jlong jPSOSize, jlong jPSOCacheDataSourceSize)
+JNI_METHOD jint Java_com_epicgames_unreal_psoservices_PSOProgramService_CompileVKGFXPSOSHM(JNIEnv* jenv, jobject thiz, jint SHMemFD, jlong jVSSize, jlong jPSSize, jlong jPSOSize, jlong jPSOCacheDataSourceSize, jfloatArray jCompilationDuration)
 {
 	ExitTest();
+
+	double CompilationStartTime = now_s();
+
 	{
 		BEGIN_TRACE("CompileVKGFXPSOSHM");
 		BEGIN_TRACE("CompileVKGFXPSOSHM_1");
@@ -1110,6 +1132,16 @@ JNI_METHOD jint Java_com_epicgames_unreal_psoservices_PSOProgramService_CompileV
 		LOG_ERROR( "Mem alloc %d bytes failed (errno %d) ", AllocSize, errno);
 	}
 
+	double CompilationDuration = now_s() - CompilationStartTime;
+
+	float* CDA = jenv->GetFloatArrayElements(jCompilationDuration, nullptr);
+	if (CDA != nullptr)
+	{
+		CDA[0] = (float)CompilationDuration;
+		jenv->ReleaseFloatArrayElements(jCompilationDuration, CDA, 0);
+	}
+
+
 	END_TRACE();
 
 	return SharedMemOutputFD;
diff --git a/Engine/Source/Runtime/Apple/MetalRHI/Private/MetalPipeline.cpp b/Engine/Source/Runtime/Apple/MetalRHI/Private/MetalPipeline.cpp
--- a/Engine/Source/Runtime/Apple/MetalRHI/Private/MetalPipeline.cpp
+++ b/Engine/Source/Runtime/Apple/MetalRHI/Private/MetalPipeline.cpp
@@ -17,6 +17,7 @@
 #include "RenderUtils.h"
 #include "Misc/ScopeRWLock.h"
 #include "HAL/PThreadEvent.h"
+#include "PSOMetrics.h"
 #include <objc/runtime.h>
 
 static int32 GMetalCacheShaderPipelines = 1;
@@ -282,7 +283,12 @@ public:
 
 				if (bCompile)
 				{
+					const double CompilationStartTime = FPlatformTime::Seconds();
 					Desc = CreateMTLRenderPipeline(Device, bSync, Key, Init, State);
+					const double CompilationDuration = FPlatformTime::Seconds() - CompilationStartTime;
+
+					GetPSOMetricsDelegate().ExecuteIfBound((float)CompilationDuration);
+
 
 					if (Desc != nullptr)
 					{
diff --git a/Engine/Source/Runtime/Engine/Private/PSOMetrics.cpp b/Engine/Source/Runtime/Engine/Private/PSOMetrics.cpp
new file mode 100644
--- /dev/null
+++ b/Engine/Source/Runtime/Engine/Private/PSOMetrics.cpp
@@ -0,0 +1,45 @@
+// Copyright Epic Games, Inc. All Rights Reserved.
+
+/*=============================================================================
+	PSOMetrics.cpp: 
+=============================================================================*/
+
+#include "PSOMetrics.h"
+#include "HAL/CriticalSection.h"
+
+namespace PSOMetrics
+{
+	FCriticalSection PSOCriticalSection;
+	float DurationSum = 0;
+	int Count = 0;
+}
+
+static FPSOMetricsEvent PSOMetricsEvent;
+FPSOMetricsEvent& GetPSOMetricsDelegate()
+{
+	if (!PSOMetricsEvent.IsBound())
+	{
+		PSOMetricsEvent.BindLambda([](float CompilationDuration)
+		{
+			// we need this scope because this will be called from RHIThread
+			FScopeLock ScopeLock(&PSOMetrics::PSOCriticalSection);
+
+			PSOMetrics::DurationSum += CompilationDuration;
+			++PSOMetrics::Count;
+		});
+	}
+
+	return PSOMetricsEvent;
+}
+
+void GetPSOCompilationMetrics(float& DurationSum, int& Count)
+{
+	// we need this scope because this will be called from the GameThread
+	FScopeLock ScopeLock(&PSOMetrics::PSOCriticalSection);
+
+	DurationSum = PSOMetrics::DurationSum;
+	Count = PSOMetrics::Count;
+
+	DurationSum = 0;
+	Count = 0;
+}
\ No newline at end of file
diff --git a/Engine/Source/Runtime/Engine/Public/PSOMetrics.h b/Engine/Source/Runtime/Engine/Public/PSOMetrics.h
new file mode 100644
--- /dev/null
+++ b/Engine/Source/Runtime/Engine/Public/PSOMetrics.h
@@ -0,0 +1,19 @@
+// Copyright Epic Games, Inc. All Rights Reserved.
+
+/*=============================================================================
+	PSOMetrics.h
+=============================================================================*/
+
+#pragma once
+
+#include "Delegates/Delegate.h"
+
+/**
+ * PSO metric delegate is called for each individual PSO compilation.
+ * At the moment only IOS and Android platforms will call this.
+ */ 
+DECLARE_DELEGATE_OneParam(FPSOMetricsEvent, float /*CompilationDuration*/);
+extern ENGINE_API FPSOMetricsEvent& GetPSOMetricsDelegate();
+
+// retrieves the current metrics and set it to zero
+extern ENGINE_API void GetPSOCompilationMetrics(float& DurationSum, int& Count);
\ No newline at end of file
diff --git a/Engine/Source/Runtime/OpenGLDrv/Private/Android/AndroidOpenGL.cpp b/Engine/Source/Runtime/OpenGLDrv/Private/Android/AndroidOpenGL.cpp
--- a/Engine/Source/Runtime/OpenGLDrv/Private/Android/AndroidOpenGL.cpp
+++ b/Engine/Source/Runtime/OpenGLDrv/Private/Android/AndroidOpenGL.cpp
@@ -16,6 +16,7 @@
 #include "String/Find.h"
 #include "String/LexFromString.h"
 #include "Android/AndroidDynamicRHI.h"
+#include "PSOMetrics.h"
 
 int32 FAndroidOpenGL::GLMajorVerion = 0;
 int32 FAndroidOpenGL::GLMinorVersion = 0;
@@ -100,6 +101,8 @@ struct FOpenGLRemoteGLProgramCompileJNI
 	jfieldID ProgramResponse_ErrorField = 0;
 	jfieldID ProgramResponse_SHMOutputHandleField = 0;
 	jfieldID ProgramResponse_CompiledBinaryField = 0;
+	jfieldID ProgramResponse_CompilationDurationField = 0;
+
 	bool bAllFound = false;
 
 	void Init(JNIEnv* Env)
@@ -138,9 +141,11 @@ struct FOpenGLRemoteGLProgramCompileJNI
 			CHECK_JNI_EXCEPTIONS(Env);
 			ProgramResponse_SHMOutputHandleField = FJavaWrapper::FindField(Env, ProgramResponseClass, "SHMOutputHandle", "I", true);
 			CHECK_JNI_EXCEPTIONS(Env);
+			ProgramResponse_CompilationDurationField = FJavaWrapper::FindField(Env, ProgramResponseClass, "CompilationDuration", "F", true);
+			CHECK_JNI_EXCEPTIONS(Env);
 		}
 
-		bAllFound = OGLServiceAccessor && DispatchProgramLink && StartRemoteProgramLink && StopRemoteProgramLink && AreProgramServicesReady && HaveServicesFailed && ProgramResponseClass && ProgramResponse_SuccessField && ProgramResponse_CompiledBinaryField && ProgramResponse_ErrorField && ProgramResponse_SHMOutputHandleField;
+		bAllFound = OGLServiceAccessor && DispatchProgramLink && StartRemoteProgramLink && StopRemoteProgramLink && AreProgramServicesReady && HaveServicesFailed && ProgramResponseClass && ProgramResponse_SuccessField && ProgramResponse_CompiledBinaryField && ProgramResponse_ErrorField && ProgramResponse_SHMOutputHandleField && ProgramResponse_CompilationDurationField;
 		UE_CLOG(!bAllFound, LogRHI, Fatal, TEXT("Failed to find JNI GL remote program compiler."));
 	}
 }OpenGLRemoteGLProgramCompileJNI;
@@ -1194,6 +1199,8 @@ TArray<uint8> FAndroidOpenGL::DispatchAndWaitForRemoteGLProgramCompile(FGraphics
 				int len = Env->GetArrayLength(*ProgramResult);
 				CompiledProgramBinary.SetNumUninitialized(len);
 				Env->GetByteArrayRegion(*ProgramResult, 0, len, reinterpret_cast<jbyte*>(CompiledProgramBinary.GetData()));
+				float CompilationDuration = (float)Env->GetFloatField(*ProgramResponseObj, OpenGLRemoteGLProgramCompileJNI.ProgramResponse_CompilationDurationField);
+				GetPSOMetricsDelegate().ExecuteIfBound(CompilationDuration);
 			}
 			else
 			{
diff --git a/Engine/Source/Runtime/VulkanRHI/Private/Android/VulkanAndroidPlatform.cpp b/Engine/Source/Runtime/VulkanRHI/Private/Android/VulkanAndroidPlatform.cpp
--- a/Engine/Source/Runtime/VulkanRHI/Private/Android/VulkanAndroidPlatform.cpp
+++ b/Engine/Source/Runtime/VulkanRHI/Private/Android/VulkanAndroidPlatform.cpp
@@ -19,6 +19,7 @@
 #include <sys/mman.h>
 #include "ProfilingDebugging/ScopedTimers.h"
 #include "Android/AndroidDynamicRHI.h"
+#include "PSOMetrics.h"
 
 #if USE_ANDROID_VULKAN_SWAPPY
 #undef VK_NO_PROTOTYPES
@@ -1390,6 +1391,7 @@ struct FVKRemoteProgramCompileJNI
 	jfieldID ProgramResponse_ErrorField = 0;
 	jfieldID ProgramResponse_SHMOutputHandleField = 0;
 	jfieldID ProgramResponse_CompiledBinaryField = 0;
+	jfieldID ProgramResponse_CompilationDurationField = 0;
 	bool bAllFound = false;
 
 	void Init(JNIEnv* Env)
@@ -1434,9 +1436,11 @@ struct FVKRemoteProgramCompileJNI
 			CHECK_JNI_EXCEPTIONS(Env);
 			ProgramResponse_SHMOutputHandleField = FJavaWrapper::FindField(Env, ProgramResponseClass, "SHMOutputHandle", "I", true);
 			CHECK_JNI_EXCEPTIONS(Env);
+			ProgramResponse_CompilationDurationField = FJavaWrapper::FindField(Env, ProgramResponseClass, "CompilationDuration", "F", true);
+			CHECK_JNI_EXCEPTIONS(Env);
 		}
 
-		bAllFound = PSOServiceAccessor && DispatchPSOCompile && DispatchPSOCompileShm && StartRemoteProgramLink && HaveServicesFailed && AreProgramServicesReady && StopRemoteProgramLink && ProgramResponseClass && ProgramResponse_SuccessField && ProgramResponse_CompiledBinaryField && ProgramResponse_ErrorField && ProgramResponse_SHMOutputHandleField;
+		bAllFound = PSOServiceAccessor && DispatchPSOCompile && DispatchPSOCompileShm && StartRemoteProgramLink && HaveServicesFailed && AreProgramServicesReady && StopRemoteProgramLink && ProgramResponseClass && ProgramResponse_SuccessField && ProgramResponse_CompiledBinaryField && ProgramResponse_ErrorField && ProgramResponse_SHMOutputHandleField && ProgramResponse_CompilationDurationField;
 		UE_CLOG(!bAllFound, LogRHI, Fatal, TEXT("Failed to find JNI Vulkan remote program compiler."));
 	}
 }VKRemoteProgramCompileJNI;
@@ -1633,6 +1637,8 @@ VkPipelineCache FVulkanAndroidPlatform::PrecompilePSO(
 			if (bSucceeded)
 			{
 				const int ProgramResultSharedHandle = Env->GetIntField(*ProgramResponseObj, VKRemoteProgramCompileJNI.ProgramResponse_SHMOutputHandleField);
+				const float ProgramResultCompilationDuration = Env->GetFloatField(*ProgramResponseObj, VKRemoteProgramCompileJNI.ProgramResponse_CompilationDurationField);
+				GetPSOMetricsDelegate().ExecuteIfBound(ProgramResultCompilationDuration);
 				if(ensure(ProgramResultSharedHandle > -1))
 				{
 					const uint32 ResultMemSize = (uint32)ASharedMemory_getSize(ProgramResultSharedHandle);
