[ASIS] Global engine. java class refactoring.
  - remove pending surfcae functionality
  - remove uniused code
  - apply coding style format
  - add unit test

#rnx

diff --git a/Engine/Build/Android/Java/src/com/epicgames/unreal/GameActivity.java.template b/Engine/Build/Android/Java/src/com/epicgames/unreal/GameActivity.java.template
--- a/Engine/Build/Android/Java/src/com/epicgames/unreal/GameActivity.java.template
+++ b/Engine/Build/Android/Java/src/com/epicgames/unreal/GameActivity.java.template
@@ -7212,12 +7212,26 @@ public class GameActivity extends $${gameActivitySuperClass}$$ implements Surfac
 			Log.debug("Cannot set \"UEEnableMemoryTracing\" environment variable");
 		}
 
+		if (!isUnderTest())
+		{
 //$${soLoadUnreal}$$
 //$${soLoadLibrary}$$
+		}
 	}
 
 	public native void nativeHandleSensorEvents(float[] tilt, float[] rotation_rate, float[] gravity, float[] acceleration);
 
 	protected static ClipboardManager clipboardManager;
 
+	private static boolean isUnderTest()
+	{
+		// Check for a specific environment variable set in tests
+		String testPackageName = System.getProperty("unreal.test.environment");
+		if (testPackageName == null)
+		{
+			return false;
+		}
+
+		return testPackageName.equals(GameApplication.getAppContext().getPackageName());
+	}
 }
