Do not attach network connection monitoring to PSOService processes

#android

diff --git a/Engine/Build/Android/Java/src/com/epicgames/unreal/GameApplication.java.template b/Engine/Build/Android/Java/src/com/epicgames/unreal/GameApplication.java.template
--- a/Engine/Build/Android/Java/src/com/epicgames/unreal/GameApplication.java.template
+++ b/Engine/Build/Android/Java/src/com/epicgames/unreal/GameApplication.java.template
@@ -28,7 +28,29 @@ public class GameApplication extends Application implements LifecycleObserver {
 
 		ProcessLifecycleOwner.get().getLifecycle().addObserver(this);
 
-		NetworkChangedManager.getInstance().initNetworkCallback(this);
+		// get the process name to see if we need network observer
+		String processName = "unknown";
+		if (android.os.Build.VERSION.SDK_INT >= 28)
+		{
+			processName = Application.getProcessName();
+		}
+		else
+		{
+			try
+			{
+				Class<?> activityThread = Class.forName("android.app.ActivityThread");
+				processName = (String)activityThread.getDeclaredMethod("currentProcessName").invoke(null);
+			}
+			catch (Exception e)
+			{
+				// ignore
+			}
+		}
+		
+		if (!processName.contains("psoprogramservice"))
+		{
+			NetworkChangedManager.getInstance().initNetworkCallback(this);
+		}
 	}
 
 	public static Context getAppContext() {
