Add GameController classes to Proguard

#android

diff --git a/Engine/Build/Android/Java/proguard-project.txt b/Engine/Build/Android/Java/proguard-project.txt
--- a/Engine/Build/Android/Java/proguard-project.txt
+++ b/Engine/Build/Android/Java/proguard-project.txt
@@ -52,7 +52,9 @@
 -keep class com.epicgames.unreal.MulticastBroadcastReceiver {
 	public *;
 }
-
+-keep class com.epicgames.unreal.GameControllerManager {
+	public *;
+}
 # Needed...
 
 -keep class com.epicgames.unreal.GameActivity$LaunchNotification {
diff --git a/Engine/Build/Android/Java/src/com/epicgames/unreal/GameControllerManager.java b/Engine/Build/Android/Java/src/com/epicgames/unreal/GameControllerManager.java
--- a/Engine/Build/Android/Java/src/com/epicgames/unreal/GameControllerManager.java
+++ b/Engine/Build/Android/Java/src/com/epicgames/unreal/GameControllerManager.java
@@ -43,7 +43,7 @@ public class GameControllerManager
 
 	private native void nativeOnZeroInputDevices();
 
-    private void scanDevices()
+    public void scanDevices()
     {
         int[] deviceIds = inputManager.getInputDeviceIds();
 
