Fix proguard issue with GameControllerManager changes

#android

diff --git a/Engine/Build/Android/Java/src/com/epicgames/unreal/GameActivity.java.template b/Engine/Build/Android/Java/src/com/epicgames/unreal/GameActivity.java.template
--- a/Engine/Build/Android/Java/src/com/epicgames/unreal/GameActivity.java.template
+++ b/Engine/Build/Android/Java/src/com/epicgames/unreal/GameActivity.java.template
@@ -610,7 +610,7 @@ public class GameActivity extends $${gameActivitySuperClass}$$ implements Surfac
 	/** App wide Multicast lock used to allow broadcast/multicast communication */
 	private WifiManager.MulticastLock MultiCastLock = null;
 
-	private GameControllerManager gameControllerManager;
+	public GameControllerManager gameControllerManager;
 
 //$${gameActivityClassAdditions}$$
 
@@ -1683,7 +1683,25 @@ public class GameActivity extends $${gameActivitySuperClass}$$ implements Surfac
 					MatchString = MatchString.toLowerCase();
 				}
 
-				if (CompareType.equals("CMP_Less"))
+				if (CompareType.equals("CMP_VEqual"))
+				{
+					boolean bSourceMatchEqual = Math.abs(SourceFloat - MatchFloat) < 0.00001f;
+					if ((bNumericOperands && !bSourceMatchEqual) || (!bNumericOperands && (Source.compareTo(MatchString) != 0)))
+					{
+						bConditionTrue = false;
+						break;
+					}
+				}
+				else if (CompareType.equals("CMP_VNotEqual"))
+				{
+					boolean bSourceMatchEqual = Math.abs(SourceFloat - MatchFloat) < 0.00001f;
+					if ((bNumericOperands && bSourceMatchEqual) || (!bNumericOperands && (Source.compareTo(MatchString) == 0)))
+					{
+						bConditionTrue = false;
+						break;
+					}
+				}
+				else if (CompareType.equals("CMP_Less"))
 				{
 					if ((bNumericOperands && (SourceFloat >= MatchFloat)) || (!bNumericOperands && (Source.compareTo(MatchString) >= 0)))
 					{
