In SplashActivity terminate early to allow activity already running at top of stack to resume

#android

diff --git a/Engine/Build/Android/Java/src/com/epicgames/unreal/SplashActivity.java b/Engine/Build/Android/Java/src/com/epicgames/unreal/SplashActivity.java
--- a/Engine/Build/Android/Java/src/com/epicgames/unreal/SplashActivity.java
+++ b/Engine/Build/Android/Java/src/com/epicgames/unreal/SplashActivity.java
@@ -44,6 +44,20 @@ public class SplashActivity extends Activity
 	protected void onCreate(Bundle savedInstanceState)
 	{
 		super.onCreate(savedInstanceState);
+
+		// check if already running
+		if (!isTaskRoot())
+		{
+			Intent intent = getIntent();
+			// see if foregrounded by launcher main action (like from home screen icon)
+			if (intent.hasCategory(Intent.CATEGORY_LAUNCHER) && intent.getAction() != null && intent.getAction().equals(Intent.ACTION_MAIN))
+			{
+				// leave previous top activity at top instead of forcing GameActivity
+				finish();
+				return;
+			}
+		}
+
 		packageName = getPackageName();
 		pm = getPackageManager();
 
