Remove usage of resolveActivity and rely on try/catch to handle unresolved intents when starting activities

#android

diff --git a/Engine/Build/Android/Java/src/com/epicgames/unreal/GameActivity.java.template b/Engine/Build/Android/Java/src/com/epicgames/unreal/GameActivity.java.template
--- a/Engine/Build/Android/Java/src/com/epicgames/unreal/GameActivity.java.template
+++ b/Engine/Build/Android/Java/src/com/epicgames/unreal/GameActivity.java.template
@@ -110,6 +110,7 @@ import android.text.Spanned;
 import android.text.method.PasswordTransformationMethod;
 import android.text.TextWatcher;
 import android.view.inputmethod.EditorInfo;
+import android.content.ActivityNotFoundException;
 import android.content.ClipboardManager;
 import android.content.ClipData;
 import android.content.ComponentCallbacks2;
@@ -4631,16 +4632,12 @@ public class GameActivity extends $${gameActivitySuperClass}$$ implements Surfac
 			BrowserIntent.addFlags(Intent.FLAG_ACTIVITY_NO_HISTORY | Intent.FLAG_ACTIVITY_CLEAR_WHEN_TASK_RESET);
 			BrowserIntent.addFlags(Intent.FLAG_ACTIVITY_NEW_TASK | Intent.FLAG_ACTIVITY_MULTIPLE_TASK);
 			
-			// make sure there is a web browser to handle the URL before trying to start activity (or may crash!)
-			if (BrowserIntent.resolveActivity(getPackageManager()) != null)
-			{
-				Log.debug("[JAVA} AndroidThunkJava_LaunchURL: Starting activity");
-				startActivity(BrowserIntent);
-			}
-			else
-			{
-				Log.debug("[JAVA} AndroidThunkJava_LaunchURL: Could not find an application to receive the URL intent");
-			}
+			Log.debug("[JAVA} AndroidThunkJava_LaunchURL: Starting activity");
+			startActivity(BrowserIntent);
+		}
+		catch (ActivityNotFoundException e)
+		{
+			Log.debug("[JAVA} AndroidThunkJava_LaunchURL: Could not find an application to receive the URL intent");
 		}
 		catch (Exception e)
 		{
@@ -6920,12 +6917,20 @@ public class GameActivity extends $${gameActivitySuperClass}$$ implements Surfac
 			intent.setComponent(new android.content.ComponentName(packageName, intentComponent));
 		}
 
-		if (intent.resolveActivity(getPackageManager()) != null)
+		try
 		{
 			startActivity(intent);
 			AndroidThunkJava_ForceQuit();
 			return true;
 		}
+		catch (ActivityNotFoundException e)
+		{
+			Log.debug("[JAVA} AndroidThunkJava_LaunchPackage: Could not start activity");
+		}
+		catch (Exception e)
+		{
+			Log.debug("[JAVA} AndroidThunkJava_LaunchPackage: Failed with exception " + e.getMessage());
+		}
 
 		return false;
 	}
