Fix for missing virtual input on Meta devices

#android

diff --git a/Engine/Build/Android/Java/src/com/epicgames/unreal/GameActivity.java.template b/Engine/Build/Android/Java/src/com/epicgames/unreal/GameActivity.java.template
--- a/Engine/Build/Android/Java/src/com/epicgames/unreal/GameActivity.java.template
+++ b/Engine/Build/Android/Java/src/com/epicgames/unreal/GameActivity.java.template
@@ -6553,7 +6553,7 @@ public class GameActivity extends $${gameActivitySuperClass}$$ implements Surfac
 	}
 	
 	//virtual keyboard input class - custom EditText
-	public class VirtualKeyboardInput extends androidx.appcompat.widget.AppCompatEditText 
+	public class VirtualKeyboardInput extends EditText
 	{
 		public VirtualKeyboardInput(Context context, AttributeSet attrs, int defStyle) 
 		{
@@ -6586,7 +6586,7 @@ public class GameActivity extends $${gameActivitySuperClass}$$ implements Surfac
 		@Override   
 		protected void onSelectionChanged(int selStart, int selEnd) 
 		{
-			super.onSelectionChanged(selStart, selEnd);
+			//super.onSelectionChanged(selStart, selEnd);
 			nativeVirtualKeyboardSendSelection(selStart, selEnd);
 		}
 
diff --git a/Engine/Source/Programs/UnrealBuildTool/Platform/Android/UEDeployAndroid.cs b/Engine/Source/Programs/UnrealBuildTool/Platform/Android/UEDeployAndroid.cs
--- a/Engine/Source/Programs/UnrealBuildTool/Platform/Android/UEDeployAndroid.cs
+++ b/Engine/Source/Programs/UnrealBuildTool/Platform/Android/UEDeployAndroid.cs
@@ -3096,6 +3096,12 @@ namespace UnrealBuildTool
 			}
 			else
 			{
+				if (bPackageForMetaQuest)
+				{
+					// needed for virtual keyboard to work on MetaQuest
+					Text.AppendLine("\t<uses-feature android:name=\"oculus.software.overlay_keyboard\" android:required=\"false\"/>");
+				}
+
 				Text.AppendLine("\t<uses-feature android:glEsVersion=\"" + AndroidToolChain.GetGLESVersion(bBuildForES31) + "\" android:required=\"true\" />");
 				Text.AppendLine("\t<uses-permission android:name=\"android.permission.INTERNET\"/>");
 				if (Configuration != "Shipping" || !bUseExternalFilesDir)
