Fix for app crash when when resuming with deep link and app is not fully initialized

diff --git a/Engine/Build/Android/Java/src/com/epicgames/unreal/GameActivity.java.template b/Engine/Build/Android/Java/src/com/epicgames/unreal/GameActivity.java.template
--- a/Engine/Build/Android/Java/src/com/epicgames/unreal/GameActivity.java.template
+++ b/Engine/Build/Android/Java/src/com/epicgames/unreal/GameActivity.java.template
@@ -291,6 +291,9 @@ public class GameActivity extends $${gameActivitySuperClass}$$ implements Surfac
 	// For Network connection change tracking
 	private NetworkConnectivityClient.Listener NetworkListener = null;
 
+	// Set to true during onNewIntent so that we can notify the app during onResume, which happens later
+	private boolean bIntentModified = false;
+
 	// From cpufeatures.h AndroidCpuFamily enum
 	private static final int ANDROID_CPU_FAMILY_UNKNOWN = 0;
 	private static final int ANDROID_CPU_FAMILY_ARM = 1;
@@ -3707,6 +3710,17 @@ public class GameActivity extends $${gameActivitySuperClass}$$ implements Surfac
 		// since calls to showSoftInput are ignored on resume so have to make sure state is reset
 		bKeyboardShowing = false;
 
+		if (bIntentModified)
+		{
+			bIntentModified = false;
+			Intent newIntent = getIntent();
+			if (newIntent != null && newIntent.getData() != null)
+			{
+				nativeHandleNewIntentData(newIntent.getData().toString());
+			}
+		}
+		
+
 //$${gameActivityOnResumeAdditions}$$
 		Log.debug("==============> GameActive.onResume complete!");
 	}
@@ -4060,11 +4074,8 @@ public class GameActivity extends $${gameActivitySuperClass}$$ implements Surfac
 	{
 		super.onNewIntent(newIntent);
 		setIntent(newIntent);
+		bIntentModified = true;
 //$${gameActivityOnNewIntentAdditions}$$
-		if (newIntent.getData() != null)
-		{
-			nativeHandleNewIntentData(newIntent.getData().toString());
-		}
 	}
 
 	public void onTrimMemory(int level)
diff --git a/Engine/Source/Runtime/Core/Private/Android/AndroidPlatformMisc.cpp b/Engine/Source/Runtime/Core/Private/Android/AndroidPlatformMisc.cpp
--- a/Engine/Source/Runtime/Core/Private/Android/AndroidPlatformMisc.cpp
+++ b/Engine/Source/Runtime/Core/Private/Android/AndroidPlatformMisc.cpp
@@ -1701,7 +1701,8 @@ static void HandleNewIntentUri_GameThread(const FString& IntentUri)
 
 void FAndroidMisc::HandleNewIntentUri(const FString& IntentUri)
 {
-	if (!IntentUri.IsEmpty())
+	// TaskGraph may be null if app was resumed before it's fully initialized, in such cases regular start will pick up the intent and send it again
+	if (!IntentUri.IsEmpty() && FTaskGraphInterface::IsRunning())
 	{
 		FFunctionGraphTask::CreateAndDispatchWhenReady([IntentUri]()
 			{
