Fixes java null exception if UECommandLine.txt is empty

#android

diff --git a/Engine/Build/Android/Java/src/com/epicgames/unreal/GameActivity.java.template b/Engine/Build/Android/Java/src/com/epicgames/unreal/GameActivity.java.template
--- a/Engine/Build/Android/Java/src/com/epicgames/unreal/GameActivity.java.template
+++ b/Engine/Build/Android/Java/src/com/epicgames/unreal/GameActivity.java.template
@@ -7343,7 +7343,7 @@ public class GameActivity extends $${gameActivitySuperClass}$$ implements Surfac
 		{
 			String Line = reader.readLine();
 			//Log.debug("Override Commandline: " + Line);
-			return Line.trim();
+			return Line != null ? Line.trim() : "";
 		}
 		catch (IOException ignored)
 		{
@@ -7353,7 +7353,7 @@ public class GameActivity extends $${gameActivitySuperClass}$$ implements Surfac
 		{
 			String Line = reader.readLine();
 			//Log.debug("Override Commandline: " + Line);
-			return Line.trim();
+			return Line != null ? Line.trim() : "";
 		}
 		catch (IOException ignored)
 		{
@@ -7363,7 +7363,7 @@ public class GameActivity extends $${gameActivitySuperClass}$$ implements Surfac
 		{
 			String Line = reader.readLine();
 			//Log.debug("APK Commandline: " + Line);
-			return Line.trim();
+			return Line != null ? Line.trim() : "";
 		}
 		catch (IOException ignored)
 		{
