Fix up support for localized notification channel names

#android

diff --git a/Engine/Build/Android/Java/src/com/epicgames/unreal/LocalNotificationReceiver.java b/Engine/Build/Android/Java/src/com/epicgames/unreal/LocalNotificationReceiver.java
--- a/Engine/Build/Android/Java/src/com/epicgames/unreal/LocalNotificationReceiver.java
+++ b/Engine/Build/Android/Java/src/com/epicgames/unreal/LocalNotificationReceiver.java
@@ -14,7 +14,8 @@ import androidx.core.app.NotificationCompat;
 public class LocalNotificationReceiver extends BroadcastReceiver
 {
 	private static final String NOTIFICATION_CHANNEL_ID = "ue4-push-notification-channel-id";
-	private static final CharSequence NOTICATION_CHANNEL_NAME = "ue4-push-notification-channel";
+	private static final String NOTICATION_CHANNEL_NAME = "ue4-push-notification-channel";
+	private static final String NOTIFICATION_CHANNEL_NAME_RESOURCE = "UEPushChannel";
 
 	public static final String KEY_LOCAL_NOTIFICATION_ID = "local-notification-ID";
 	public static final String KEY_LOCAL_NOTIFICATION_TITLE = "local-notification-title";
@@ -36,7 +37,7 @@ public class LocalNotificationReceiver extends BroadcastReceiver
 			return;
 		}
 
-		// Open UE4 app if clicked
+		// Open UE5 app if clicked
 		Intent notificationIntent = new Intent(context, GameActivity.class);
 
 		// launch if closed but resume if running
@@ -70,7 +71,8 @@ public class LocalNotificationReceiver extends BroadcastReceiver
 				NotificationChannel channel = notificationManager.getNotificationChannel(NOTIFICATION_CHANNEL_ID);
 				if (channel == null)
 				{
-					channel = new NotificationChannel(NOTIFICATION_CHANNEL_ID, NOTICATION_CHANNEL_NAME, NotificationManager.IMPORTANCE_DEFAULT);
+					CharSequence channelName = getLocalizedResource(context.getApplicationContext(), NOTIFICATION_CHANNEL_NAME_RESOURCE, NOTICATION_CHANNEL_NAME);
+					channel = new NotificationChannel(NOTIFICATION_CHANNEL_ID, channelName, NotificationManager.IMPORTANCE_DEFAULT);
 					channel.enableVibration(true);
 					channel.enableLights(true);
 					notificationManager.createNotificationChannel(channel);
@@ -106,4 +108,78 @@ public class LocalNotificationReceiver extends BroadcastReceiver
 		}
 		return notificationIconID;
 	}
+
+	private static int getResourceId(Context context, String VariableName, String ResourceName, String PackageName)
+	{
+		try {
+			return context.getResources().getIdentifier(VariableName, ResourceName, PackageName);
+		}
+		catch (Exception e) {
+			e.printStackTrace();
+			return -1;
+		} 
+	}
+
+	private static String getResourceStringOrDefault(Context context, String PackageName, String ResourceName, String DefaultString)
+	{
+		int resourceId = getResourceId(context, ResourceName, "string", PackageName);
+		return (resourceId < 1) ? DefaultString : context.getResources().getString(resourceId);
+	}
+
+	public static String getLocalizedResource(Context context, String baseName, String defaultValue)
+	{
+		String packageName = context.getApplicationContext().getPackageName();
+
+		// first try using Android default locale handling
+		String checkValue = getResourceStringOrDefault(context, packageName, baseName, "");
+		if (!"".equals(checkValue))
+		{
+			return checkValue;
+		}
+
+		// check for default English (if it doesn't exist, no other locales do either
+		String englishValue = getResourceStringOrDefault(context, packageName, baseName + "_en", "");
+		if ("".equals(englishValue))
+		{
+			return defaultValue;
+		}
+
+		// now try exact locale
+		String locale = java.util.Locale.getDefault().toString().replace("-", "_");
+		checkValue = getResourceStringOrDefault(context, packageName, baseName + "_" + locale, "");
+		if (!"".equals(checkValue))
+		{
+			return checkValue;
+		}
+
+		// next try without underscore if there is one (ex. es_MX -> es)
+		if (locale.contains("_"))
+		{
+			String partialLocale = locale.substring(0, locale.indexOf("_"));
+			checkValue = getResourceStringOrDefault(context, packageName, baseName + "_" + partialLocale, "");
+			if (!"".equals(checkValue))
+			{
+				return checkValue;
+			}
+		}
+
+		// finally try some remappings
+		switch (locale)
+		{
+			case "es_419": locale = "es_MX"; break;
+			case "zh": locale = "zh_CN"; break;
+			case "zh_Hant": locale = "zh_TW"; break;
+			case "zh_rTW": locale = "zh_TW"; break;
+			case "zh_rHK": locale = "zh_TW"; break;
+			case "zh_rMO": locale = "zh_TW"; break;
+		}
+		checkValue = getResourceStringOrDefault(context, packageName, baseName + "_" + locale, "");
+		if (!"".equals(checkValue))
+		{
+			return checkValue;
+		}
+
+		// finally, fall back to English
+		return englishValue;
+	}
 }
diff --git a/Engine/Plugins/Online/Android/AndroidFetchBackgroundDownload/Source/AndroidFetchBackgroundDownload/Java/download/UEDownloadWorker.java b/Engine/Plugins/Online/Android/AndroidFetchBackgroundDownload/Source/AndroidFetchBackgroundDownload/Java/download/UEDownloadWorker.java
--- a/Engine/Plugins/Online/Android/AndroidFetchBackgroundDownload/Source/AndroidFetchBackgroundDownload/Java/download/UEDownloadWorker.java
+++ b/Engine/Plugins/Online/Android/AndroidFetchBackgroundDownload/Source/AndroidFetchBackgroundDownload/Java/download/UEDownloadWorker.java
@@ -733,7 +733,14 @@ public class UEDownloadWorker extends UEWorker implements DownloadProgressListen
 				NotificationChannel Channel = notificationManager.getNotificationChannel(Description.NotificationChannelID);
 				if (Channel == null)
 				{
-					Channel = new NotificationChannel(Description.NotificationChannelID, Description.NotificationChannelName, Description.NotificationChannelImportance);
+					String ChannelName = Description.NotificationChannelName;
+					if (ChannelName.equals("ue-downloadworker-channel"))
+					{
+						// check for an override for default channel name
+						ChannelName = com.epicgames.unreal.LocalNotificationReceiver.getLocalizedResource(context, "UEBackgroundChannel", ChannelName);
+					}
+					
+					Channel = new NotificationChannel(Description.NotificationChannelID, ChannelName, Description.NotificationChannelImportance);
 					notificationManager.createNotificationChannel(Channel);
 				}
 			}
