Add build configuration to Android's configrules variables.

#rnx

diff --git a/Engine/Build/Android/Java/src/com/epicgames/unreal/GameActivity.java.template b/Engine/Build/Android/Java/src/com/epicgames/unreal/GameActivity.java.template
--- a/Engine/Build/Android/Java/src/com/epicgames/unreal/GameActivity.java.template
+++ b/Engine/Build/Android/Java/src/com/epicgames/unreal/GameActivity.java.template
@@ -1894,6 +1894,7 @@ public class GameActivity extends $${gameActivitySuperClass}$$ implements Surfac
 		variables.put("SRC_SDKLevel", Integer.toString(android.os.Build.VERSION.SDK_INT));
 		variables.put("supportsFloatingPointRenderTargets", bSupportsFloatingPointRenderTargets ? "true" : "false");
 		variables.put("TextureFormats", TextureFormats);
+		variables.put("SRC_BuildConfig", BuildConfiguration);
 
 		// get the status bar and navigation bar heights (only works for 5.0+)
 		int statusBarHeight = 0;
