Android: Move ConsoleCmdReceiver registration parameters to reusable method.

#rnx

diff --git a/Engine/Build/Android/Java/src/com/epicgames/unreal/ConsoleCmdReceiver.java b/Engine/Build/Android/Java/src/com/epicgames/unreal/ConsoleCmdReceiver.java
--- a/Engine/Build/Android/Java/src/com/epicgames/unreal/ConsoleCmdReceiver.java
+++ b/Engine/Build/Android/Java/src/com/epicgames/unreal/ConsoleCmdReceiver.java
@@ -16,30 +16,55 @@
 
 package com.epicgames.unreal;
 
+import java.util.Objects;
+
+import android.annotation.SuppressLint;
 import android.content.BroadcastReceiver;
 import android.content.Context;
+import android.content.ContextWrapper;
 import android.content.Intent;
+import android.content.IntentFilter;
+import android.os.Build;
+
+import androidx.annotation.NonNull;
 
 public class ConsoleCmdReceiver extends BroadcastReceiver {
 
 	private GameActivity gameActivity;
-	
+
 	// Constructor...
 	public ConsoleCmdReceiver(GameActivity InGameActivity) {
 		gameActivity = InGameActivity;
 	}
-		
+
 	@Override
-    public void onReceive(Context context, Intent intent) {
-			// example usage
-			// adb shell "am broadcast -a android.intent.action.RUN -e cmd 'stat fps'"
-			String action = intent.getAction();
-			if (action != null && action == Intent.ACTION_RUN) {
-				String cmd = intent.getStringExtra("cmd");
-				if (cmd != null) {
-					gameActivity.nativeConsoleCommand(cmd);
-				}
+	public void onReceive(Context context, Intent intent) {
+		// example usage
+		// adb shell "am broadcast -a android.intent.action.RUN -e cmd 'stat fps'"
+		String action = intent.getAction();
+		if (Objects.equals(action, Intent.ACTION_RUN)) {
+			String cmd = intent.getStringExtra("cmd");
+			if (cmd != null) {
+				gameActivity.nativeConsoleCommand(cmd);
 			}
-    }
+		}
+	}
+
+	/**
+	 * Registers this to the given context using Unreal Engine conventions.  In the current version, that implies receiving the RUN action publicly.
+	 * You need to call `contextWrapper.unregisterReceiver(mConsoleCmdReceiver)` when the context is torn down.
+	 * @param contextWrapper The context to which this will be registered.
+	 */
+	public void register(@NonNull ContextWrapper contextWrapper) {
+		if (Build.VERSION.SDK_INT >= Build.VERSION_CODES.TIRAMISU) {
+			contextWrapper.registerReceiver(this, new IntentFilter(Intent.ACTION_RUN), Context.RECEIVER_EXPORTED);
+		} else {
+			compatRegister(contextWrapper);
+		}
+	}
 
+	@SuppressLint("UnspecifiedRegisterReceiverFlag")
+	private void compatRegister(@NonNull ContextWrapper contextWrapper) {
+		contextWrapper.registerReceiver(this, new IntentFilter(Intent.ACTION_RUN));
+	}
 }
diff --git a/Engine/Build/Android/Java/src/com/epicgames/unreal/GameActivity.java.template b/Engine/Build/Android/Java/src/com/epicgames/unreal/GameActivity.java.template
--- a/Engine/Build/Android/Java/src/com/epicgames/unreal/GameActivity.java.template
+++ b/Engine/Build/Android/Java/src/com/epicgames/unreal/GameActivity.java.template
@@ -661,7 +661,7 @@ public class GameActivity extends $${gameActivitySuperClass}$$ implements Surfac
 			// Create console command broadcast listener
 			Log.debug( "Creating console command broadcast listener");
 			consoleCmdReceiver = new ConsoleCmdReceiver(this);
-			registerReceiver(consoleCmdReceiver, new IntentFilter(Intent.ACTION_RUN), Context.RECEIVER_EXPORTED);
+			consoleCmdReceiver.register(this);
 		}
 //$${gameActivityOnStartAdditions}$$
 		Log.debug("==================================> Inside onStart function in GameActivity");
