[Android VK] Fix texts sent partially and text fields being updated incorrectly after pressing a character while the VK is closing:
- SlateEditableTextLayout: Force update the text before commiting if the VK text changed
- AndroidJNI: Make the "Done" process async so it waits for all the pending text changes
- GameActivity.java:
1. Flush pending text change operations on "Done" action (if any)
2. Avoid updating the text if the VK was requested to close

#tests Tested on an Android Iterative build
#rnx

diff --git a/Engine/Build/Android/Java/src/com/epicgames/unreal/GameActivity.java.template b/Engine/Build/Android/Java/src/com/epicgames/unreal/GameActivity.java.template
--- a/Engine/Build/Android/Java/src/com/epicgames/unreal/GameActivity.java.template
+++ b/Engine/Build/Android/Java/src/com/epicgames/unreal/GameActivity.java.template
@@ -415,6 +415,9 @@ public class GameActivity extends $${gameActivitySuperClass}$$ implements Surfac
 
 	//handler for virtual keyboard show/hide events
 	private Handler virtualKeyboardHandler;
+	
+	// Keep a reference to the VK text change runnable to flush it when pressing the "Done" VK button if necessary  
+	private Runnable virtualKeyboardTextChangeRunnable = null;
 
 	// Keep a reference to the main content view so we can bring up the virtual keyboard without an editbox
 	private View mainView;
@@ -6763,6 +6766,13 @@ public class GameActivity extends $${gameActivitySuperClass}$$ implements Surfac
 				switch(result) 
 				{
 				case EditorInfo.IME_ACTION_DONE:
+					// Flush the pending text changes (if any)
+					if (virtualKeyboardTextChangeRunnable != null)
+					{
+						virtualKeyboardHandler.removeCallbacks(virtualKeyboardTextChangeRunnable);
+						virtualKeyboardTextChangeRunnable.run();
+						virtualKeyboardTextChangeRunnable = null;
+					}
 					AndroidThunkJava_HideVirtualKeyboardInput();
 					nativeVirtualKeyboardSendKey(KeyEvent.KEYCODE_ENTER);
 	                return true;
@@ -6796,33 +6806,46 @@ public class GameActivity extends $${gameActivitySuperClass}$$ implements Surfac
 					//delay the "set empty string" event and wait for a second call
 					if(charSequence.length() == 0)
 					{
-						virtualKeyboardHandler.postDelayed(new Runnable()
+						virtualKeyboardTextChangeRunnable = new Runnable()
 						{
 							public void run()
 							{
-								String message = newVirtualKeyboardInput.getText().toString();
-								nativeVirtualKeyboardChanged(message);
+								virtualKeyboardTextChangeRunnable = null;
+								
+								// Process text changes only if the keyboard was not requested to close
+								if (lastVirtualKeyboardCommand != VirtualKeyboardCommand.VK_CMD_HIDE)
+								{
+									String message = newVirtualKeyboardInput.getText().toString();
+									nativeVirtualKeyboardChanged(message);
+								}
 							}
-						}, 100);
+						};
+						virtualKeyboardHandler.postDelayed(virtualKeyboardTextChangeRunnable, 100);
 					}
 					else
 					{
-						virtualKeyboardHandler.postDelayed(new Runnable()
+						virtualKeyboardTextChangeRunnable = new Runnable()
 						{
 							public void run()
 							{
-								String message = newVirtualKeyboardInput.getText().toString();
-								//#jira UE-50645 Carriage returns can be pasted into single line UMG fields on Android
-								//oddly enough, we have to use events/filters to control the EditText's copy/paste behaviour
-								if(newVirtualKeyboardInput.getMaxLines() == 1 && message.contains("\n"))
+								virtualKeyboardTextChangeRunnable = null;
+
+								// Process text changes only if the keyboard was not requested to close
+								if (lastVirtualKeyboardCommand != VirtualKeyboardCommand.VK_CMD_HIDE)
 								{
-									message = message.replaceAll("\n" , " ");
-									newVirtualKeyboardInput.setText(message);
-						
+									String message = newVirtualKeyboardInput.getText().toString();
+									//#jira UE-50645 Carriage returns can be pasted into single line UMG fields on Android
+									//oddly enough, we have to use events/filters to control the EditText's copy/paste behaviour
+									if(newVirtualKeyboardInput.getMaxLines() == 1 && message.contains("\n"))
+									{
+										message = message.replaceAll("\n" , " ");
+										newVirtualKeyboardInput.setText(message);
+									}
+									nativeVirtualKeyboardChanged(message);
 								}
-								nativeVirtualKeyboardChanged(message);
 							}
-						}, 100);
+						};
+						virtualKeyboardHandler.postDelayed(virtualKeyboardTextChangeRunnable, 100);
 					}
 				}
 				downgradeEasyCorrectionSpans();
diff --git a/Engine/Source/Runtime/Launch/Private/Android/AndroidJNI.cpp b/Engine/Source/Runtime/Launch/Private/Android/AndroidJNI.cpp
--- a/Engine/Source/Runtime/Launch/Private/Android/AndroidJNI.cpp
+++ b/Engine/Source/Runtime/Launch/Private/Android/AndroidJNI.cpp
@@ -1235,11 +1235,21 @@ JNI_METHOD void Java_com_epicgames_unreal_GameActivity_nativeVirtualKeyboardChan
 
 JNI_METHOD void Java_com_epicgames_unreal_GameActivity_nativeVirtualKeyboardSendKey(JNIEnv* jenv, jobject thiz, jint keyCode)
 {
-	FDeferredAndroidMessage Message;
+	if (FTaskGraphInterface::IsRunning())
+	{
+		// Dispatch when ready to delay the submit when the text is being changed by another task (e.g Java_com_epicgames_unreal_GameActivity_nativeVirtualKeyboardChanged)
+		FFunctionGraphTask::CreateAndDispatchWhenReady([keyCode = keyCode, CurrentVirtualKeyboardWidget = VirtualKeyboardWidget]()
+		{
+			if (CurrentVirtualKeyboardWidget.IsValid() && CurrentVirtualKeyboardWidget == VirtualKeyboardWidget)
+			{
+				FDeferredAndroidMessage Message;
 
-	Message.messageType = MessageType_KeyDown;
-	Message.KeyEventData.keyId = keyCode;
-	FAndroidInputInterface::DeferMessage(Message);
+				Message.messageType = MessageType_KeyDown;
+				Message.KeyEventData.keyId = keyCode;
+				FAndroidInputInterface::DeferMessage(Message);
+			}
+		}, TStatId(), NULL, ENamedThreads::GameThread);
+	}
 }
 
 JNI_METHOD void Java_com_epicgames_unreal_GameActivity_nativeVirtualKeyboardSendSelection(JNIEnv* jenv, jobject thiz, jint selStart, jint selEnd)
diff --git a/Engine/Source/Runtime/Slate/Private/Widgets/Text/SlateEditableTextLayout.cpp b/Engine/Source/Runtime/Slate/Private/Widgets/Text/SlateEditableTextLayout.cpp
--- a/Engine/Source/Runtime/Slate/Private/Widgets/Text/SlateEditableTextLayout.cpp
+++ b/Engine/Source/Runtime/Slate/Private/Widgets/Text/SlateEditableTextLayout.cpp
@@ -1651,6 +1651,9 @@ bool FSlateEditableTextLayout::HandleCarriageReturn(bool isRepeat)
 		// Always clear the local undo chain on commit.
 		ClearUndoStates();
 
+		// Make sure to update the text if there is any VK change that was not processed yet
+		UpdateTextChangedByVirtualKeyboard();
+
 		const FText EditedText = GetEditableText();
 
 		// When enter is pressed text is committed.  Let anyone interested know about it.
@@ -3399,10 +3402,8 @@ bool FSlateEditableTextLayout::ComputeVolatility() const
 		|| LineHeightPercentage.IsBound();
 }
 
-void FSlateEditableTextLayout::Tick(const FGeometry& AllottedGeometry, const double InCurrentTime, const float InDeltaTime)
+void FSlateEditableTextLayout::UpdateTextChangedByVirtualKeyboard()
 {
-	check(IsInGameThread());
-
 	if (bTextChangedByVirtualKeyboard)
 	{
 		SetEditableText(VirtualKeyboardText);
@@ -3410,6 +3411,13 @@ void FSlateEditableTextLayout::Tick(const FGeometry& AllottedGeometry, const dou
 		OwnerWidget->OnTextChanged(GetEditableText());
 		bTextChangedByVirtualKeyboard = false;
 	}
+}
+
+void FSlateEditableTextLayout::Tick(const FGeometry& AllottedGeometry, const double InCurrentTime, const float InDeltaTime)
+{
+	check(IsInGameThread());
+
+	UpdateTextChangedByVirtualKeyboard();
 
 	if (bTextCommittedByVirtualKeyboard)
 	{
diff --git a/Engine/Source/Runtime/Slate/Public/Widgets/Text/SlateEditableTextLayout.h b/Engine/Source/Runtime/Slate/Public/Widgets/Text/SlateEditableTextLayout.h
--- a/Engine/Source/Runtime/Slate/Public/Widgets/Text/SlateEditableTextLayout.h
+++ b/Engine/Source/Runtime/Slate/Public/Widgets/Text/SlateEditableTextLayout.h
@@ -431,6 +431,8 @@ private:
 
 	/** Called when the active context menu is closed */
 	SLATE_API void OnContextMenuClosed(TSharedRef<IMenu> Menu);
+	
+	void UpdateTextChangedByVirtualKeyboard();
 
 private:
 	/** Virtual keyboard handler for an editable text layout */
