Overlap system info collection with loadlibraries
Overlap configrules download with loadlibraries
Add caching to configrules (optional CONFIGRULES_CACHE_KEY for encryption)

#android

diff --git a/Engine/Build/Android/Java/src/com/epicgames/unreal/GameActivity.java.template b/Engine/Build/Android/Java/src/com/epicgames/unreal/GameActivity.java.template
--- a/Engine/Build/Android/Java/src/com/epicgames/unreal/GameActivity.java.template
+++ b/Engine/Build/Android/Java/src/com/epicgames/unreal/GameActivity.java.template
@@ -42,6 +42,7 @@ import java.util.List;
 import java.util.Map;
 import java.util.HashMap;
 import java.util.ArrayList;
+import java.util.LinkedList;
 import java.text.DecimalFormat;
 import java.text.DecimalFormatSymbols;
 import java.text.BreakIterator;
@@ -49,6 +50,7 @@ import java.util.zip.Inflater;
 import java.util.Stack;
 import java.util.Timer;
 import java.util.TimerTask;
+import java.lang.Thread;
 import java.lang.Runnable;
 import java.util.zip.CRC32;
 import java.util.zip.Checksum;
@@ -58,6 +60,7 @@ import java.io.BufferedWriter;
 import java.io.BufferedInputStream;
 import java.io.BufferedOutputStream;
 import java.io.ByteArrayInputStream;
+import java.io.ByteArrayOutputStream;
 import java.io.Reader;
 import java.io.Writer;
 import java.io.File;
@@ -71,8 +74,10 @@ import java.io.InputStream;
 import java.io.InputStreamReader;
 import java.io.InputStream;
 import java.io.InputStreamReader;
+import java.io.ObjectInputStream;
 import java.io.OutputStream;
 import java.io.OutputStreamWriter;
+import java.io.ObjectOutputStream;
 import java.io.StringWriter;
 
 import android.annotation.TargetApi;
@@ -355,6 +360,13 @@ public class GameActivity extends $${gameActivitySuperClass}$$ implements Surfac
 	protected static Bundle _extrasBundle;
 	protected boolean bOnCreateCalled = false;
 
+	Thread ProcessSystemInfoThread = null;
+	protected long ConfigRulesCRC = 0;
+	protected long ConfigRulesVarsCRC = 0;
+	protected Map<String, String> ConfigRulesVars = null;
+	private LinkedList<String> SystemInfoLogs = null;
+	Map<String, String> CpuInfo = null;
+
 	private HandlerThread memoryHandlerThread;
 	private Handler memoryHandler;
 	private Runnable memoryRunnable;
@@ -455,6 +467,8 @@ public class GameActivity extends $${gameActivitySuperClass}$$ implements Surfac
 	private static final int PLAY_SERVICES_DIALOG_ID = 1;
 	
 	protected static String appPackageName = "";
+	protected static String ProjectName = "";
+	protected static String AppType = "";
 
 	/** Check to see if we have all the files */
 	private boolean HasAllFiles = false;
@@ -472,6 +486,7 @@ public class GameActivity extends $${gameActivitySuperClass}$$ implements Surfac
 	private boolean InitCompletedOK = false;
 
 	private boolean bConfigRulesError = false;
+	private int ConfigRulesVersion = -1;
 	private String ConfigRulesStatus = "UNSET";
 
 	/** Dialog for force game dialog */
@@ -1115,17 +1130,38 @@ public class GameActivity extends $${gameActivitySuperClass}$$ implements Surfac
 	// returns InputStream to most up to date version of configrules
 	private InputStream getConfigRulesStream(String ProjectName)
 	{
+		ConfigRulesVersion = -1;
+		ConfigRulesCRC = 0;
 		ConfigRulesStatus = "";
 
+		long assetsCRC = 0;
+		long internalCRC = 0;
+		long externalCRC = 0;
+
 		byte[] configBytesAssets = getByteArrayFromAssets("configrules.bin.png");
 		byte[] configBytesInternal = getByteArrayFromFile(InternalFilesDir + "configrules.bin.png");
 		byte[] configBytesExternal = getByteArrayFromFile(android.os.Environment.getExternalStorageDirectory().getAbsolutePath() + "/UnrealGame/" + ProjectName + "/configrules.bin.png");
 
+		if (configBytesAssets != null)
+		{
+			CRC32 crc = new CRC32();
+			crc.update(configBytesAssets);
+			assetsCRC = crc.getValue();
+		}
+
 		if (configBytesInternal != null)
 		{
 			CRC32 crc = new CRC32();
 			crc.update(configBytesInternal);
-			Log.debug(String.format("CRC32: %08X", crc.getValue()));
+			internalCRC = crc.getValue();
+			Log.debug(String.format("CRC32: %08X", internalCRC));
+		}
+
+		if (configBytesExternal != null)
+		{
+			CRC32 crc = new CRC32();
+			crc.update(configBytesExternal);
+			externalCRC = crc.getValue();
 		}
 
 		int assetsVersion = getConfigRulesVersion(configBytesAssets);
@@ -1142,37 +1178,53 @@ public class GameActivity extends $${gameActivitySuperClass}$$ implements Surfac
 		InputStream result = null;
 		if (CommandLineFlags.containsKey("noconfigrulesdownload"))
 		{
+			ConfigRulesVersion = assetsVersion;
+			ConfigRulesCRC = assetsCRC;
 			result = getConfigRulesDecompressed(configBytesAssets);
 		}
 		else if (assetsVersion >= internalVersion && assetsVersion >= externalVersion)
 		{
+			ConfigRulesVersion = assetsVersion;
+			ConfigRulesCRC = assetsCRC;
 			result = getConfigRulesDecompressed(configBytesAssets);
 		}
 		else if (internalVersion >= assetsVersion && internalVersion >= externalVersion)
 		{
+			ConfigRulesVersion = internalVersion;
+			ConfigRulesCRC = internalCRC;
 			result = getConfigRulesDecompressed(configBytesInternal);
 			if (result == null && assetsVersion != -1)
 			{
 				Log.debug("using assets rule: ");
 				File internalconfigFile = new File(InternalFilesDir + "configrules.bin.png");
 				internalconfigFile.delete();
+				ConfigRulesVersion = assetsVersion;
+				ConfigRulesCRC = assetsCRC;
 				result = getConfigRulesDecompressed(configBytesAssets);
 			}
 		}
 		else
 		{
+			ConfigRulesVersion = externalVersion;
+			ConfigRulesCRC = externalCRC;
 			result = getConfigRulesDecompressed(configBytesExternal);
 			if (result == null && internalVersion != -1)
 			{
+				ConfigRulesVersion = internalVersion;
+				ConfigRulesCRC = internalCRC;
 				result = getConfigRulesDecompressed(configBytesInternal);
 			}
 			if (result == null && assetsVersion != -1)
 			{
+				ConfigRulesVersion = assetsVersion;
+				ConfigRulesCRC = assetsCRC;
 				result = getConfigRulesDecompressed(configBytesAssets);
 			}
 		}
 		if (result == null)
 		{
+			ConfigRulesVersion = -1;
+			ConfigRulesCRC = 0;
 			ConfigRulesStatus = "failed to find config rules. ";
 			bConfigRulesError = true;
 		}
@@ -1232,7 +1284,10 @@ public class GameActivity extends $${gameActivitySuperClass}$$ implements Surfac
 			}
 			catch (IOException ie)
 			{
-				Log.debug("failed to read /proc/cpuinfo: " + ie);
+				if (SystemInfoLogs != null)
+				{
+					SystemInfoLogs.add("failed to read /proc/cpuinfo: " + ie);
+				}
 			}
 			if (reader != null)
 			{
@@ -1247,7 +1302,10 @@ public class GameActivity extends $${gameActivitySuperClass}$$ implements Surfac
 		}
 		else
 		{
-			Log.debug("failed to open /proc/cpuinfo!");
+			if (SystemInfoLogs != null)
+			{
+				SystemInfoLogs.add("failed to open /proc/cpuinfo!");
+			}
 		}
 		return cpuInfo;
 	}
@@ -1283,7 +1341,10 @@ public class GameActivity extends $${gameActivitySuperClass}$$ implements Surfac
 		}
 		catch (Exception e)
 		{
-			Log.debug("Unable to use getprop: " + e);
+			if (SystemInfoLogs != null)
+			{
+				SystemInfoLogs.add("Unable to use getprop: " + e);
+			}
 		}
 		if (reader != null)
 		{
@@ -1757,17 +1818,19 @@ public class GameActivity extends $${gameActivitySuperClass}$$ implements Surfac
 		return bConditionTrue;
 	}
 
-	private boolean processSystemInfo(String ProjectName, String PackageName)
+	private void processSystemInfo()
 	{
+		SystemInfoLogs = new LinkedList<String>();
+
 		// check OpenGL information
 		getGraphicsInformation();
 
 		// get CPU / chipset information
-		Map<String, String> cpuInfo = getCpuInfo();
+		CpuInfo = getCpuInfo();
 		int processorCount = 1;
-		if (cpuInfo.containsKey("processorCount"))
+		if (CpuInfo.containsKey("processorCount"))
 		{
-			processorCount = Integer.parseInt(cpuInfo.get("processorCount"));
+			processorCount = Integer.parseInt(CpuInfo.get("processorCount"));
 		}
 
 		// check for libhoudini.so to identify Intel Atom
@@ -1778,9 +1841,9 @@ public class GameActivity extends $${gameActivitySuperClass}$$ implements Surfac
 		boolean bIsARM64 = false;
 		boolean bHasNEON = false;
 		String processor = "";
-		if (cpuInfo.containsKey("Processor"))
+		if (CpuInfo.containsKey("Processor"))
 		{
-			processor = cpuInfo.get("Processor");
+			processor = CpuInfo.get("Processor");
 		}
 		if (processor.contains("aarch64"))
 		{
@@ -1788,19 +1851,19 @@ public class GameActivity extends $${gameActivitySuperClass}$$ implements Surfac
 			bHasNEON = true;
 		}
 
-        String abilist = getProp("ro.product.cpu.abilist");
-        if (abilist.contains("arm64-v8a"))
-        {
-            bIsARM64 = true;
-        }
+		String abilist = getProp("ro.product.cpu.abilist");
+		if (abilist.contains("arm64-v8a"))
+		{
+			bIsARM64 = true;
+		}
 
 		// check each processor featureset
 		for (int id=0; id < processorCount; id++)
 		{
 			String key = id + ":Features";
-			if (cpuInfo.containsKey(key))
+			if (CpuInfo.containsKey(key))
 			{
-				String features = cpuInfo.get(key);
+				String features = CpuInfo.get(key);
 
 				if (features.contains("neon"))
 				{
@@ -1818,29 +1881,29 @@ public class GameActivity extends $${gameActivitySuperClass}$$ implements Surfac
 		// try to get chip info
 		// use getprop ro.hardware as fallback (not useful on Qualcomm, just returns qcom)
 		String roHardware = getProp("ro.hardware");
-		Log.debug("ro.hardware: " + roHardware);
+		SystemInfoLogs.add("ro.hardware: " + roHardware);
 		if (roHardware.equals("qcom"))
 		{
 			// use ro.hardware.chipname for Qualcomm
 			String roHardwareChipname = getProp("ro.hardware.chipname");
-			Log.debug("ro.hardware.chipname: " + roHardwareChipname);
+			SystemInfoLogs.add("ro.hardware.chipname: " + roHardwareChipname);
 			roHardware = roHardwareChipname;
 			if (roHardware.isEmpty())
 			{
-				// provide GPU as a final fallback in case cpuInfo does not contain Hardware
+				// provide GPU as a final fallback in case CpuInfo does not contain Hardware
 				roHardware = OpenGLDevice;
 			}
 		}
 		String hardware = roHardware;
-		if (cpuInfo.containsKey("Hardware"))
+		if (CpuInfo.containsKey("Hardware"))
 		{
-			hardware = cpuInfo.get("Hardware");
+			hardware = CpuInfo.get("Hardware");
 		}
-		Log.debug("Hardware: " + hardware);
+		SystemInfoLogs.add("Hardware: " + hardware);
 		String roSOCModel = getProp("ro.soc.model");
 		if (!roSOCModel.equals(""))
 		{
-			Log.debug("ro.soc.model: " + roSOCModel);
+			SystemInfoLogs.add("ro.soc.model: " + roSOCModel);
 
 			// validate either we only have ro.soc.model, or that previous hardware value in agreement
 			// encountered a device which had ro.hardware and /proc/cpuinfo correct but ro.soc.model set differently
@@ -1852,19 +1915,20 @@ public class GameActivity extends $${gameActivitySuperClass}$$ implements Surfac
 			}
 			else
 			{
-				Log.debug("Warning: hardware does not match ro.soc.model, assuming hardware is correct");
+				SystemInfoLogs.add("Warning: hardware does not match ro.soc.model, assuming hardware is correct");
 			}
 		}
 
 		// register variables for the program to access
-		Map<String, String> variables = new HashMap<String, String>();
+		ConfigRulesVars = new HashMap<String, String>();
+		Map<String, String> variables = ConfigRulesVars;
 
 		ActivityManager activityManager = (ActivityManager)getSystemService(Context.ACTIVITY_SERVICE);
 		ActivityManager.MemoryInfo memoryInfo = new ActivityManager.MemoryInfo();
 		activityManager.getMemoryInfo(memoryInfo);
 
 		int totalMemory = (int)(memoryInfo.totalMem / 1024 / 1024);
-		Log.debug("Memory: " + totalMemory + " MB");
+		SystemInfoLogs.add("Memory: " + totalMemory + " MB");
 		variables.put("memory", Integer.toString(totalMemory));
 
 		//variables.put("Profile", "Android");
@@ -1933,6 +1997,28 @@ public class GameActivity extends $${gameActivitySuperClass}$$ implements Surfac
 		variables.put("screenWidth", Integer.toString(displaySize.x));
 		variables.put("screenHeight", Integer.toString(displaySize.y));
 
+		// calculate a CRC to use for cache invalidation
+		CRC32 crc32 = new CRC32();
+		for (Map.Entry<String, String> entry : variables.entrySet())
+		{
+			crc32.update(entry.getKey().getBytes());
+			crc32.update(entry.getValue().getBytes());
+		}
+		ConfigRulesVarsCRC = crc32.getValue();
+	}
+
+	private boolean runConfigRules(Map<String, String> variables, String ProjectName, String PackageName)
+	{
+		// dump any SystemInfo logs
+		if (SystemInfoLogs != null)
+		{
+			for (String logline : SystemInfoLogs)
+			{
+				Log.debug(logline);
+			}
+			SystemInfoLogs = null;
+		}
+
 		// set early restart if -noearlyrestart not on commandline (configrules can clear this to disable or set to false)
 		if (!CommandLineFlags.containsKey("noearlyrestart"))
 		{
@@ -1945,6 +2031,7 @@ public class GameActivity extends $${gameActivitySuperClass}$$ implements Surfac
 		// commandline flag may be used to skip dialog if not a shipping build
 		boolean bIgnoreDialog = !nativeIsShippingBuild() && CommandLineFlags.containsKey("noconfigrulesdialog");
 
+		// NOTE: the "earlyrestart" variable and anything added to variables by <gameActivityPreConfigRulesParseAdditions> not included in ConfigRulesVarsCRC to invalidate cache!
 //$${gameActivityPreConfigRulesParseAdditions}$$
 		
 		// read the configuration file (configrules.bin.png)
@@ -1965,11 +2052,109 @@ public class GameActivity extends $${gameActivitySuperClass}$$ implements Surfac
 			}
 		}
 
+		Log.debug("$$$ configrules: start parsing");
+
+		// check if we have a valid cached parsed version already
+		int configFlags = 0;
+		boolean bWriteCache = true;
+		boolean bClearCache = false;
+		java.security.Key cacheKey = null;
+		try
+		{
+			Field keyField = this.getClass().getDeclaredField("CONFIGRULES_CACHE_KEY");
+			cacheKey = generateAESKey((String)keyField.get(this));
+		}
+		catch (Exception e)
+		{
+		}
+		if (reader != null)
+		{
+			String configCacheFilename = InternalFilesDir + "configrules.cache";
+			byte[] cacheFileContents = getByteArrayFromFile(configCacheFilename);
+			if (cacheFileContents != null)
+			{
+				String currentHeader = String.format("%d, %08X, %08X", ConfigRulesVersion, ConfigRulesCRC, ConfigRulesVarsCRC);
+				Log.debug("checking against: " + currentHeader);
+
+				int headerSize = 4 + 8 + 8 + 4;
+				int sizeEncrypted = cacheFileContents.length - headerSize;
+				if (sizeEncrypted > 0)
+				{
+					ByteBuffer buffer = ByteBuffer.wrap(cacheFileContents);
+					int cacheVersion = buffer.getInt();
+					long cacheRulesCRC = buffer.getLong();
+					long cacheVarsCRC = buffer.getLong();
+					int cacheFlags = buffer.getInt();
+					Log.debug("read: " + String.format("%d, %08X, %08X", cacheVersion, cacheRulesCRC, cacheVarsCRC));
+					if (cacheVersion == ConfigRulesVersion && cacheRulesCRC == ConfigRulesCRC && cacheVarsCRC == ConfigRulesVarsCRC)
+					{
+						byte[] decrypted = null;
+						if (cacheKey != null)
+						{
+							try
+							{
+								Cipher cipher = Cipher.getInstance("AES/ECB/PKCS5Padding");
+								cipher.init(Cipher.DECRYPT_MODE, cacheKey);
+								decrypted = cipher.doFinal(cacheFileContents, headerSize, sizeEncrypted);
+							}
+							catch (Exception e)
+							{
+								Log.debug("$$$ decrypt failure: " + e.toString());
+							}
+						}
+						else
+						{
+							decrypted = Arrays.copyOfRange(cacheFileContents, headerSize, cacheFileContents.length);
+						}
+
+						if (decrypted != null)
+						{
+							try
+							{
+								ObjectInputStream objStream = new ObjectInputStream(new ByteArrayInputStream(decrypted));
+
+								variables = (HashMap<String, String>)objStream.readObject();
+								ConfigRulesVars = variables;
+								Log.debug("using configrules cache");
+
+								if ((cacheFlags & 2) > 0 && CpuInfo != null)
+								{
+									Log.debug("cpuInfo:");
+									for (Map.Entry<String, String> entry : CpuInfo.entrySet())
+									{
+										Log.debug("  " + entry.getKey() + " = " + entry.getValue());
+									}
+								}
+								if ((cacheFlags & 1) > 0)
+								{
+									Log.debug("variables:");
+									for (Map.Entry<String, String> entry : variables.entrySet())
+									{
+										Log.debug("  " + entry.getKey() + " = " + entry.getValue());
+									}
+								}
+
+								reader = null;
+								stream.close();
+								stream = null;
+								bWriteCache = false;
+								configFlags = cacheFlags;
+							}
+							catch (Exception e)
+							{
+								Log.debug("$$$ deserialize failure: " + e.toString());
+							}
+						}
+					}
+				}
+			}
+		}
+
 		// process the rules
 		if (reader != null)
 		{
 			// convert to lowercase once to check against chipset key
-			String hardwareCheck = hardware.toLowerCase();
+			String hardwareCheck = variables.get("hardware").toLowerCase();
 			
 			try
 			{
@@ -2246,17 +2431,22 @@ public class GameActivity extends $${gameActivitySuperClass}$$ implements Surfac
 					// check if requested to dump cpu info to the log
 					if (variables.containsKey("dumpcpuinfo"))
 					{
+						configFlags |= 2;
 						variables.remove("dumpcpuinfo");
-						Log.debug("cpuInfo:");
-						for (Map.Entry<String, String> entry : cpuInfo.entrySet())
+						if (CpuInfo != null)
 						{
-							Log.debug("  " + entry.getKey() + " = " + entry.getValue());
+							Log.debug("cpuInfo:");
+							for (Map.Entry<String, String> entry : CpuInfo.entrySet())
+							{
+								Log.debug("  " + entry.getKey() + " = " + entry.getValue());
+							}
 						}
 					}
 
 					// check if requested to dump variables to the log
 					if (variables.containsKey("dumpvars"))
 					{
+						configFlags |= 1;
 						variables.remove("dumpvars");
 						Log.debug("variables:");
 						for (Map.Entry<String, String> entry : variables.entrySet())
@@ -2270,6 +2460,8 @@ public class GameActivity extends $${gameActivitySuperClass}$$ implements Surfac
 					boolean hasWarning = variables.containsKey("warning");
 					if (hasError || hasWarning)
 					{
+						bClearCache = true;
+
 						ForceExitLink = "";
 						if (variables.containsKey("link"))
 						{
@@ -2375,6 +2567,7 @@ public class GameActivity extends $${gameActivitySuperClass}$$ implements Surfac
 			ForceExitMessage = getResourceStringOrDefault(PackageName, "CR_Info_DeviceNotSupportedYet", "This device is in supported yet.");
 			SaveError = "ConfigRules_Problem";
 			bForceGameEndWithError = true;
+			bClearCache = true;
 		}
 		if (stream != null)
 		{
@@ -2387,6 +2580,60 @@ public class GameActivity extends $${gameActivitySuperClass}$$ implements Surfac
 			}
 		}
 
+		// warnings or errors will clear the cache
+		if (bClearCache)
+		{
+			File configCacheFile = new File(InternalFilesDir + "configrules.cache");
+			configCacheFile.delete();
+			bWriteCache = false;
+		}
+
+		// write cache
+		if (bWriteCache)
+		{
+			File configCacheFile = new File(InternalFilesDir + "configrules.cache");
+			try (ByteArrayOutputStream hashStream = new ByteArrayOutputStream())
+			{
+				ObjectOutputStream objStream = new ObjectOutputStream(hashStream);
+				objStream.writeObject(variables);
+				objStream.close();
+				byte[] hashData = hashStream.toByteArray();
+				hashStream.close();
+
+				int headerSize = 4 + 8 + 8 + 4;
+				byte[] headerBytes = new byte[headerSize];
+				ByteBuffer buffer = ByteBuffer.wrap(headerBytes);
+				buffer.putInt(ConfigRulesVersion);
+				buffer.putLong(ConfigRulesCRC);
+				buffer.putLong(ConfigRulesVarsCRC);
+				buffer.putInt(configFlags);
+
+				FileOutputStream fileStream = new FileOutputStream(configCacheFile);
+				fileStream.write(headerBytes);
+
+				if (cacheKey == null)
+				{
+					fileStream.write(hashData);
+				}
+				else
+				{
+					Cipher cipher = Cipher.getInstance("AES/ECB/PKCS5Padding");
+					cipher.init(Cipher.ENCRYPT_MODE, cacheKey);
+					byte[] encrypted = cipher.doFinal(hashData, 0, hashData.length);
+					fileStream.write(encrypted);
+				}
+				fileStream.close();
+			}
+			catch (IOException e)
+			{
+				configCacheFile.delete();
+			}
+			catch (Exception e)
+			{
+				configCacheFile.delete();
+			}
+		}
+
 		// clear Profile if OverrideProfile was not left set as well
 		if (variables.containsKey("Profile"))
 		{
@@ -2530,9 +2777,117 @@ public class GameActivity extends $${gameActivitySuperClass}$$ implements Surfac
 		Log.debug("LIFECYCLE: GameActivity.onCreate() called!, this=" + this);
 		
 		_extrasBundle = getIntent().getExtras();
-		
+
+		activityContext = getActivityOfContext(this);
+		assert activityContext != null;
+
+		// Retrieve version code, name, and target SDK
+		try 
+		{
+			PackageInfo packageInfo = getPackageManager().getPackageInfo(getPackageName(), 0);
+			VersionCode = packageInfo.versionCode;
+			VersionName = packageInfo.versionName;
+			targetSdkVersion = packageInfo.applicationInfo.targetSdkVersion;
+		}
+		catch (Exception e)
+		{
+			Log.debug("Error accessing packageInfo: " + e.getMessage());
+		}
+
+		// read some metadata from AndroidManifest.xml early
+		appPackageName = getPackageName();
+		ProjectName = appPackageName.substring(ProjectName.lastIndexOf('.') + 1);
+		try {
+			ApplicationInfo ai = getPackageManager().getApplicationInfo(appPackageName, PackageManager.GET_META_DATA);
+			Bundle bundle = ai.metaData;
+			_bundle = ai.metaData;
+
+			if ((ai.flags & ApplicationInfo.FLAG_DEBUGGABLE) == 0) 
+			{
+				IsForDistribution = true;
+			}
+
+			// Get the project name from AndroidManifest.xml
+			if (bundle.containsKey("com.epicgames.unreal.GameActivity.ProjectName"))
+			{
+				ProjectName = bundle.getString("com.epicgames.unreal.GameActivity.ProjectName");
+				Log.debug( "Found ProjectName = " + ProjectName);
+			}
+			else
+			{
+				Log.debug( "Did not find ProjectName, using package name = " + ProjectName);
+			}
+
+			// Get the application type from AndroidManifest.xml
+			if (bundle.containsKey("com.epicgames.unreal.GameActivity.AppType"))
+			{
+				AppType = bundle.getString("com.epicgames.unreal.GameActivity.AppType");
+				Log.debug( "Found AppType = " + AppType);
+			}
+			else
+			{
+				Log.debug( "Did not find AppType, using default application type");
+			}
+
+			if (bundle.containsKey("com.epicgames.unreal.GameActivity.bShouldHideUI"))
+			{
+				ShouldHideUI = bundle.getBoolean("com.epicgames.unreal.GameActivity.bShouldHideUI");
+				Log.debug( "UI hiding set to " + ShouldHideUI);
+			}
+			else
+			{
+				Log.debug( "UI hiding not found. Leaving as " + ShouldHideUI);
+			}
+
+			if (bundle.containsKey("com.epicgames.unreal.GameActivity.BuildConfiguration"))
+			{
+				BuildConfiguration = bundle.getString("com.epicgames.unreal.GameActivity.BuildConfiguration");
+				Log.debug( "BuildConfiguration set to " + BuildConfiguration);
+			}
+			else
+			{
+				Log.debug( "BuildConfiguration not found" );
+			}
+
+			if (bundle.containsKey("com.epicgames.unreal.GameActivity.bSupportsVulkan"))
+			{
+				bSupportsVulkan = bundle.getBoolean("com.epicgames.unreal.GameActivity.bSupportsVulkan");
+				Log.debug( "SupportsVulkan set to " + bSupportsVulkan);
+			}
+
+			if (bundle.containsKey("com.epicgames.unreal.GameActivity.EngineBranch"))
+			{
+				EngineBranch = bundle.getString("com.epicgames.unreal.GameActivity.EngineBranch");
+			}
+
+			if (bundle.containsKey("com.epicgames.unreal.GameActivity.ProjectVersion"))
+			{
+				ProjectVersion = bundle.getString("com.epicgames.unreal.GameActivity.ProjectVersion");
+			}
+		}
+		catch (NameNotFoundException e)
+		{
+			Log.debug( "Failed to load meta-data: NameNotFound: " + e.getMessage());
+		}
+		catch (NullPointerException e)
+		{
+			Log.debug( "Failed to load meta-data: NullPointer: " + e.getMessage());
+		}
+
+		// get file paths
 		AssetManagerReference = this.getAssets();
-		
+		InternalFilesDir = getFilesDir().getAbsolutePath() + "/";
+		ExternalFilesDir = Objects.requireNonNull(getExternalFilesDir(null)).getAbsolutePath() + "/";
+
+		ProcessSystemInfoThread = new Thread(new Runnable() {
+            public void run() {
+				processSystemInfo();
+			}
+		});
+		ProcessSystemInfoThread.setName("ProcessSysInfo");
+		ProcessSystemInfoThread.start();
+
+//$${gameActivityOnCreateBeforeLoadLibrariesAdditions}$$
 		LoadLibraries();
 
 		// Bug with certain devices calling onCreate twice, even with singleTask set and System.exit called in onDestroy.
@@ -2551,11 +2906,7 @@ public class GameActivity extends $${gameActivitySuperClass}$$ implements Surfac
 		
 		_activity = this;
 		bOnCreateCalled = true;
-		activityContext = getActivityOfContext(this);
-		assert activityContext != null;
 		
-		InternalFilesDir = getFilesDir().getAbsolutePath() + "/";
-		ExternalFilesDir = Objects.requireNonNull(getExternalFilesDir(null)).getAbsolutePath() + "/";
 //$${gameActivityOnCreateBeginningAdditions}$$
 		Logger.RegisterCallback(this);
 
@@ -2711,18 +3062,6 @@ public class GameActivity extends $${gameActivitySuperClass}$$ implements Surfac
 			}
 		}
 
-		// Retrieve version code, name, and target SDK
-		try 
-		{
-			PackageInfo packageInfo = getPackageManager().getPackageInfo(getPackageName(), 0);
-			VersionCode = packageInfo.versionCode;
-			VersionName = packageInfo.versionName;
-			targetSdkVersion = packageInfo.applicationInfo.targetSdkVersion;
-		}
-		catch (Exception e)
-		{
-			Log.debug("Error accessing packageInfo: " + e.getMessage());
-		}
 
 		// Check for target sdk.  If 23 or higher then warn that permission handling may mean features don't work if user denies them.
 		if (!isStandalone && ANDROID_BUILD_VERSION >= 23 && targetSdkVersion >= 23) //23 is the API level (Marshmallow) where runtime permission handling is available
@@ -2991,19 +3330,8 @@ public class GameActivity extends $${gameActivitySuperClass}$$ implements Surfac
 		}
 
 		// Read metadata from AndroidManifest.xml
-		appPackageName = getPackageName();
-		String ProjectName = getPackageName();
-		ProjectName = ProjectName.substring(ProjectName.lastIndexOf('.') + 1);
-		String AppType = "";
 		try {
-			ApplicationInfo ai = getPackageManager().getApplicationInfo(getPackageName(), PackageManager.GET_META_DATA);
-			Bundle bundle = ai.metaData;
-			_bundle = bundle;
-
-			if ((ai.flags & ApplicationInfo.FLAG_DEBUGGABLE) == 0) 
-			{
-				IsForDistribution = true;
-			}
+			Bundle bundle = _bundle;
 
 			// Get the preferred depth buffer size from AndroidManifest.xml
 			if (bundle.containsKey("com.epicgames.unreal.GameActivity.DepthBufferPreference"))
@@ -3038,28 +3366,6 @@ public class GameActivity extends $${gameActivitySuperClass}$$ implements Surfac
 				Log.debug( "Did not find bPackageDataInsideApk, using default.");
 			}
 
-			// Get the project name from AndroidManifest.xml
-			if (bundle.containsKey("com.epicgames.unreal.GameActivity.ProjectName"))
-			{
-				ProjectName = bundle.getString("com.epicgames.unreal.GameActivity.ProjectName");
-				Log.debug( "Found ProjectName = " + ProjectName);
-			}
-			else
-			{
-				Log.debug( "Did not find ProjectName, using package name = " + ProjectName);
-			}
-			
-			// Get the application type from AndroidManifest.xml
-			if (bundle.containsKey("com.epicgames.unreal.GameActivity.AppType"))
-			{
-				AppType = bundle.getString("com.epicgames.unreal.GameActivity.AppType");
-				Log.debug( "Found AppType = " + AppType);
-			}
-			else
-			{
-				Log.debug( "Did not find AppType, using default application type");
-			}
-			
 			if (bundle.containsKey("com.epicgames.unreal.GameActivity.bHasOBBFiles"))
 			{
 				HasOBBFiles = bundle.getBoolean("com.epicgames.unreal.GameActivity.bHasOBBFiles") ? 1 : 0;
@@ -3082,16 +3388,6 @@ public class GameActivity extends $${gameActivitySuperClass}$$ implements Surfac
 				Log.debug( "Did not find bVerifyOBBOnStartUp, using default.");
 			}
 
-			if(bundle.containsKey("com.epicgames.unreal.GameActivity.bShouldHideUI"))
-			{
-				ShouldHideUI = bundle.getBoolean("com.epicgames.unreal.GameActivity.bShouldHideUI");
-				Log.debug( "UI hiding set to " + ShouldHideUI);
-			}
-			else
-			{
-				Log.debug( "UI hiding not found. Leaving as " + ShouldHideUI);
-			}
-
 			if (SplashScreenLaunch == false && android.os.Build.VERSION.SDK_INT >= 28)
 			{
 				if(bundle.containsKey("com.epicgames.unreal.GameActivity.bUseDisplayCutout"))
@@ -3105,16 +3401,6 @@ public class GameActivity extends $${gameActivitySuperClass}$$ implements Surfac
 				}
 			}
 
-			if(bundle.containsKey("com.epicgames.unreal.GameActivity.BuildConfiguration"))
-			{
-				BuildConfiguration = bundle.getString("com.epicgames.unreal.GameActivity.BuildConfiguration");
-				Log.debug( "BuildConfiguration set to " + BuildConfiguration);
-			}
-			else
-			{
-				Log.debug( "BuildConfiguration not found" );
-			}
-
 			if (bundle.containsKey("com.epicgames.unreal.GameActivity.bUseExternalFilesDir"))
             {
                 UseExternalFilesDir = bundle.getBoolean("com.epicgames.unreal.GameActivity.bUseExternalFilesDir");
@@ -3141,44 +3427,41 @@ public class GameActivity extends $${gameActivitySuperClass}$$ implements Surfac
 				Log.debug( "AllowIMU set to " + bAllowIMU);
 			}
 
-			if(bundle.containsKey("com.epicgames.unreal.GameActivity.bSupportsVulkan"))
-			{
-				bSupportsVulkan = bundle.getBoolean("com.epicgames.unreal.GameActivity.bSupportsVulkan");
-				Log.debug( "SupportsVulkan set to " + bSupportsVulkan);
-			}
-
-			if(bundle.containsKey("com.epicgames.unreal.GameActivity.EngineBranch"))
-			{
-				EngineBranch = bundle.getString("com.epicgames.unreal.GameActivity.EngineBranch");
-			}
-
-			if(bundle.containsKey("com.epicgames.unreal.GameActivity.ProjectVersion"))
-			{
-				ProjectVersion = bundle.getString("com.epicgames.unreal.GameActivity.ProjectVersion");
-			}
-
 //$${gameActivityReadMetadataAdditions}$$
 		}
-		catch (NameNotFoundException e)
-		{
-			Log.debug( "Failed to load meta-data: NameNotFound: " + e.getMessage());
-		}
+//		catch (NameNotFoundException e)
+//		{
+//			Log.debug( "Failed to load meta-data: NameNotFound: " + e.getMessage());
+//		}
 		catch (NullPointerException e)
 		{
 			Log.debug( "Failed to load meta-data: NullPointer: " + e.getMessage());
 		}
+		catch (Exception e)
+		{
+			Log.debug( "Failed to load meta-data: " + e.getMessage());
+		}
 
 		// parse the commandline
 		ParseCommandline(ProjectName, UseExternalFilesDir);
 
 //$${gameActivityBeforeConfigRulesAppliedAdditions}$$
 
-		// run configuration rules
-		if (!processSystemInfo(ProjectName, appPackageName))
+		// run configuration rules (wait for SystemInfoLogs fetch to complete first)
+		try
+		{
+			ProcessSystemInfoThread.join();
+		}
+		catch (Exception e)
+		{
+		}
+		Log.debug("$$$ configrules: run config rules");
+		if (!runConfigRules(ConfigRulesVars, ProjectName, appPackageName))
 		{
 			// configuration failed (show error)
 			return;
 		}
+		Log.debug("$$$ configrules: finish config rules");
 
 		String APKPath = getPackageResourcePath();
 //$${gameActivityOverrideAPKOBBPackaging}$$
@@ -3349,7 +3632,7 @@ public class GameActivity extends $${gameActivitySuperClass}$$ implements Surfac
 
 		// Now okay for event handler to be set up on native side
 		//	nativeResumeMainInit();
-				
+
 		// Try to establish a connection to Google Play
 		// AndroidThunkJava_GooglePlayConnect();
 
@@ -7473,11 +7756,21 @@ public class GameActivity extends $${gameActivitySuperClass}$$ implements Surfac
 		// We need to decide on enabling memory tracing before loading native code.
 		try
 		{
-			Bundle Bundle = GameApplication.getAppContext().getPackageManager().getApplicationInfo(getPackageName(), PackageManager.GET_META_DATA).metaData;
-			
-			if (!"Shipping".equals(Bundle.getString("com.epicgames.unreal.GameActivity.BuildConfiguration")))
-			{ 
-				EnableMemoryTracing(AppendDebugCommandLine(GetNotShippingCommandLine(Bundle, CommandLine)));
+			if (_bundle != null)
+			{
+				if (!"Shipping".equals(BuildConfiguration))
+				{ 
+					EnableMemoryTracing(AppendDebugCommandLine(GetNotShippingCommandLine(_bundle, CommandLine)));
+				}
+			}
+			else
+			{
+				Bundle Bundle = GameApplication.getAppContext().getPackageManager().getApplicationInfo(getPackageName(), PackageManager.GET_META_DATA).metaData;
+
+				if (!"Shipping".equals(Bundle.getString("com.epicgames.unreal.GameActivity.BuildConfiguration")))
+				{ 
+					EnableMemoryTracing(AppendDebugCommandLine(GetNotShippingCommandLine(Bundle, CommandLine)));
+				}
 			}
 		}
 		catch (ErrnoException | NameNotFoundException e)
diff --git a/Engine/Source/Programs/UnrealBuildTool/Platform/Android/UEDeployAndroid.cs b/Engine/Source/Programs/UnrealBuildTool/Platform/Android/UEDeployAndroid.cs
--- a/Engine/Source/Programs/UnrealBuildTool/Platform/Android/UEDeployAndroid.cs
+++ b/Engine/Source/Programs/UnrealBuildTool/Platform/Android/UEDeployAndroid.cs
@@ -6130,6 +6130,7 @@ import java.util.Collection;
 				{ "//$${gameActivityImplementsAdditions}$$", UPL.ProcessPluginNode(NDKArch, "gameActivityImplementsAdditions", "")},
 				{ "//$${gameActivityClassAdditions}$$", UPL.ProcessPluginNode(NDKArch, "gameActivityClassAdditions", GameActivityClassAdditionsDefault)},
 				{ "//$${gameActivityReadMetadataAdditions}$$", UPL.ProcessPluginNode(NDKArch, "gameActivityReadMetadataAdditions", "")},
+				{ "//$${gameActivityOnCreateBeforeLoadLibrariesAdditions}$$", UPL.ProcessPluginNode(NDKArch, "gameActivityOnCreateBeforeLoadLibrariesAdditions", "")},
 				{ "//$${gameActivityOnCreateBeginningAdditions}$$", UPL.ProcessPluginNode(NDKArch, "gameActivityOnCreateBeginningAdditions", "")},
 				{ "//$${gameActivityOnCreateAdditions}$$", UPL.ProcessPluginNode(NDKArch, "gameActivityOnCreateAdditions", "")},
 				{ "//$${gameActivityOnCreateFinalAdditions}$$", UPL.ProcessPluginNode(NDKArch, "gameActivityOnCreateFinalAdditions", "")},
