New BitmapRenderer for MediaPlayer14
Notes:
- Vulkan path is using FVYUConvertPS.
- The GL path creates texture with GL_TEXTURE_EXTERNAL_OES and the driver will auto convert YUV color space.
- Need further tests with other video encodings like ffmpeg.

diff --git a/Engine/Build/Android/Java/proguard-project.txt b/Engine/Build/Android/Java/proguard-project.txt
--- a/Engine/Build/Android/Java/proguard-project.txt
+++ b/Engine/Build/Android/Java/proguard-project.txt
@@ -55,6 +55,15 @@
 -keep class com.epicgames.unreal.GameControllerManager {
 	public *;
 }
+-keep class com.epicgames.unreal.BitmapRendererImageReader {
+	public *;
+}
+-keep class com.epicgames.unreal.BitmapRendererLegacy {
+	public *;
+}
+-keep class com.epicgames.unreal.BitmapRenderer$FrameData {
+	public *;
+}
 # Needed...
 
 -keep class com.epicgames.unreal.GameActivity$LaunchNotification {
diff --git a/Engine/Build/Android/Java/src/com/epicgames/unreal/BitmapRendererImageReader.java b/Engine/Build/Android/Java/src/com/epicgames/unreal/BitmapRendererImageReader.java
new file mode 100644
--- /dev/null
+++ b/Engine/Build/Android/Java/src/com/epicgames/unreal/BitmapRendererImageReader.java
@@ -0,0 +1,201 @@
+// Copyright Epic Games, Inc. All Rights Reserved.
+
+package com.epicgames.unreal;
+
+import android.graphics.ImageFormat;
+import android.graphics.SurfaceTexture;
+import android.hardware.HardwareBuffer;
+import android.media.Image;
+import android.media.ImageReader;
+import android.os.Handler;
+import android.util.Log;
+import android.view.Surface;
+
+import java.nio.Buffer;
+
+interface BitmapRenderer
+{
+    class UpdateParam
+    {
+        int destTexture;
+    }
+
+    // unified Frame Data encapsulating all kinds of outputs
+    class FrameData
+    {
+        public void release()
+        {
+            if (hardwareBuffer != null)
+            {
+                hardwareBuffer.close();
+                hardwareBuffer = null;
+            }
+
+            if (image != null)
+            {
+                image.close();
+                image = null;
+            }
+        }
+
+        // output of BitmapRendererImageReader, Image/HardwareBuffer
+        public Image image = null;
+        public HardwareBuffer hardwareBuffer = null;
+        public float UScale = 0.0f;
+        public float UOffset = 0.0f;
+        public float VScale = 0.0f;
+        public float VOffset = 0.0f;
+
+        // output of BitmapRendererLegacy, byte array
+        public Buffer binaryBuffer = null;
+        // output of BitmapRendererLegacy, result of rendering to an external texture
+        public boolean frameReady = false;
+    }
+
+    boolean isValid();
+    boolean resolutionChanged();
+    void setSize(int width, int height);
+    Surface getSurface();
+    SurfaceTexture getSurfaceTexture();
+    FrameData updateFrameData(UpdateParam param);
+    void release();
+}
+
+class BitmapRendererImageReader implements BitmapRenderer, ImageReader.OnImageAvailableListener
+{
+    private static final String TAG = "BitmapRenderer";
+    private ImageReader mImageReader = null;
+    private Handler mImageReaderHandler = null;
+    private int mTextureWidth;
+    private int mTextureHeight;
+    private boolean mTextureSizeChanged = true;
+    private Surface mSurface;
+    private boolean mUseVulkanRenderer;
+
+    public BitmapRendererImageReader(boolean swizzlePixels, boolean isVulkan, int width, int height, int maxQueueLength)
+    {
+        mUseVulkanRenderer = isVulkan;
+
+        mTextureWidth = width;
+        mTextureHeight = height;
+
+        //mImageReader = ImageReader.newInstance(width, 1, ImageFormat.YUV_420_888, maxQueueLength, HardwareBuffer.USAGE_GPU_SAMPLED_IMAGE);
+		mImageReader = ImageReader.newInstance(width, height, ImageFormat.PRIVATE, maxQueueLength, HardwareBuffer.USAGE_GPU_SAMPLED_IMAGE);
+
+        mSurface = mImageReader.getSurface();
+
+        mImageReaderHandler = new android.os.Handler(android.os.Looper.getMainLooper());
+        mImageReader.setOnImageAvailableListener(this, mImageReaderHandler);
+    }
+
+    public void onImageAvailable(ImageReader reader)
+    {
+        //nativeSignalSurfaceReadEvent(mParentHandle);
+    }
+
+    public synchronized FrameData updateFrameData(UpdateParam param)
+    {
+        if (mImageReader == null)
+        {
+            // Can't update if there's no surface to update into.
+            return null;
+        }
+
+        Image image;
+        try
+        {
+            // Get image data from Surface / Queue...
+            image = mImageReader.acquireNextImage();
+        }
+        catch (Exception e)
+        {
+            Log.d(TAG, String.format("updateFrameData: ImageReader.acquireNextImage() threw: %s", e.toString()));
+            return null;
+        }
+
+        if (image == null)
+        {
+            //Log.d(TAG, "updateFrameData requested, but no images queued...");
+            return null;
+        }
+
+        HardwareBuffer hwBuffer = image.getHardwareBuffer();
+        if (hwBuffer == null)
+        {
+            Log.d(TAG, "updateFrameData: Could not get HardwareBuffer from Image!");
+            return null;
+        }
+
+        FrameData frameData = new FrameData();
+        frameData.image = image;
+        frameData.hardwareBuffer = hwBuffer;
+
+        android.graphics.Rect CropRect = frameData.image.getCropRect();
+        frameData.UScale = (float)CropRect.width() / mTextureWidth;
+        frameData.UOffset = (float)CropRect.left / mTextureWidth;
+        frameData.VScale = (float)CropRect.height() / mTextureHeight;
+        frameData.VOffset = (float)CropRect.top / mTextureHeight;
+
+		//Log.d(TAG, "updateFrameData: new Image aquired. size: " + CropRect.width() + "x" + CropRect.height() + "; view: " + width + "x" + height);
+        return frameData;
+    }
+
+    public boolean isValid()
+    {
+        return mImageReader != null;
+    }
+
+    public Surface getSurface()
+    {
+        return mSurface;
+    }
+
+    public SurfaceTexture getSurfaceTexture()
+    {
+        // TODO
+        throw new UnsupportedOperationException();
+    }
+
+    public synchronized void setSize(int width, int height)
+    {
+        if (width != mTextureWidth || height != mTextureHeight)
+        {
+            mTextureWidth = width;
+            mTextureHeight = height;
+
+            // TODO - Reinitialize the image reader
+
+            mTextureSizeChanged = true;
+        }
+    }
+
+    public boolean resolutionChanged()
+    {
+        boolean changed;
+        synchronized(this)
+        {
+            changed = mTextureSizeChanged;
+            mTextureSizeChanged = false;
+        }
+        return changed;
+    }
+
+    public synchronized void release()
+    {
+        if (mImageReader != null)
+        {
+            mImageReader.setOnImageAvailableListener(null, null);
+        }
+
+        if (null != mSurface)
+        {
+            mSurface.release();
+            mSurface = null;
+        }
+        if (mImageReader != null)
+        {
+            mImageReader.close();
+            mImageReader = null;
+        }
+    }
+}
diff --git a/Engine/Build/Android/Java/src/com/epicgames/unreal/BitmapRendererLegacy.java b/Engine/Build/Android/Java/src/com/epicgames/unreal/BitmapRendererLegacy.java
new file mode 100644
--- /dev/null
+++ b/Engine/Build/Android/Java/src/com/epicgames/unreal/BitmapRendererLegacy.java
@@ -0,0 +1,832 @@
+// Copyright Epic Games, Inc. All Rights Reserved.
+
+package com.epicgames.unreal;
+
+import android.graphics.SurfaceTexture;
+import android.opengl.*;
+import android.util.Log;
+
+/*
+	All this internal surface view does is manage the
+	offscreen bitmap that the media player decoding can
+	render into for eventual extraction to the UE4 buffers.
+*/
+class BitmapRendererLegacy implements BitmapRenderer, SurfaceTexture.OnFrameAvailableListener
+{
+    private static final String TAG = "BitmapRenderer";
+    private java.nio.Buffer mFrameData = null;
+    private int mLastFramePosition = -1;
+    private android.graphics.SurfaceTexture mSurfaceTexture = null;
+    private int mTextureWidth = -1;
+    private int mTextureHeight = -1;
+    private android.view.Surface mSurface = null;
+    private boolean mFrameAvailable = false;
+    private int mTextureID = -1;
+    private int mFBO = -1;
+    private int mBlitVertexShaderID = -1;
+    private int mBlitFragmentShaderID = -1;
+    private float[] mTransformMatrix = new float[16];
+    private boolean mTriangleVerticesDirty = true;
+    private boolean mTextureSizeChanged = true;
+    private boolean mUseOwnContext = true;
+    private boolean mVulkanRenderer = false;
+    private boolean mSwizzlePixels = false;
+
+    private int GL_TEXTURE_EXTERNAL_OES = 0x8D65;
+
+    private EGLDisplay mEglDisplay;
+    private EGLContext mEglContext;
+    private EGLSurface mEglSurface;
+
+    private EGLDisplay mSavedDisplay;
+    private EGLContext mSavedContext;
+    private EGLSurface mSavedSurfaceDraw;
+    private EGLSurface mSavedSurfaceRead;
+
+    private boolean mCreatedEGLDisplay = false;
+
+    public BitmapRendererLegacy(boolean swizzlePixels, boolean isVulkan)
+    {
+        mSwizzlePixels = swizzlePixels;
+        mVulkanRenderer = isVulkan;
+
+        mEglSurface = EGL14.EGL_NO_SURFACE;
+        mEglContext = EGL14.EGL_NO_CONTEXT;
+        mUseOwnContext = true;
+
+        if (mVulkanRenderer)
+        {
+            mSwizzlePixels = true;
+        }
+        else
+        {
+            String RendererString = GLES20.glGetString(GLES20.GL_RENDERER);
+
+            // Do not use shared context if Adreno before 400 or on older Android than Marshmallow
+            if (RendererString.contains("Adreno (TM) "))
+            {
+                java.util.regex.Pattern pattern = java.util.regex.Pattern.compile("Adreno \\(TM\\) (\\d*)");
+                java.util.regex.Matcher matcher = pattern.matcher(RendererString);
+                if (matcher.find() && matcher.groupCount() >= 1)
+                {
+                    int AdrenoVersion = Integer.parseInt(matcher.group(1));
+                    if (AdrenoVersion < 400 || android.os.Build.VERSION.SDK_INT < 22)
+                    {
+                        Log.d(TAG, "MediaPlayer14: disabled shared GL context on " + RendererString);
+                        mUseOwnContext = false;
+                    }
+                }
+            }
+        }
+
+        if (mUseOwnContext)
+        {
+            initContext();
+            saveContext();
+            makeCurrent();
+            initSurfaceTexture();
+            restoreContext();
+        }
+        else
+        {
+            initSurfaceTexture();
+        }
+    }
+
+    private void initContext()
+    {
+        mEglDisplay = EGL14.EGL_NO_DISPLAY;
+        EGLContext shareContext = EGL14.EGL_NO_CONTEXT;
+
+        int majorver[] = new int[] { 0 };
+        int minorver[] = new int[] { 0 };
+        if (!mVulkanRenderer)
+        {
+            mEglDisplay = EGL14.eglGetCurrentDisplay();
+            shareContext = EGL14.eglGetCurrentContext();
+
+            if (android.os.Build.VERSION.SDK_INT >= 18 &&
+                    EGL14.eglQueryContext(mEglDisplay, shareContext, EGLExt.EGL_CONTEXT_MAJOR_VERSION_KHR, majorver, 0) &&
+                    EGL14.eglQueryContext(mEglDisplay, shareContext, EGLExt.EGL_CONTEXT_MINOR_VERSION_KHR, minorver, 0))
+            {
+                Log.d(TAG, "MediaPlayer14: Existing GL context is version " + majorver[0] + "." + minorver[0]);
+            }
+            else
+                // on some devices eg Galaxy S6, the above fails but we do get EGL14.EGL_CONTEXT_CLIENT_VERSION=3
+                if (EGL14.eglQueryContext(mEglDisplay, shareContext, EGL14.EGL_CONTEXT_CLIENT_VERSION, majorver, 0))
+                {
+                    Log.d(TAG, "MediaPlayer14: Existing GL context is version " + majorver[0]);
+                }
+                else
+                {
+                    Log.d(TAG, "MediaPlayer14: Existing GL context version not detected");
+                }
+        }
+        else
+        {
+            mEglDisplay = EGL14.eglGetDisplay(EGL14.EGL_DEFAULT_DISPLAY);
+            if (mEglDisplay == EGL14.EGL_NO_DISPLAY)
+            {
+                Log.d(TAG, "unable to get EGL14 display");
+                return;
+            }
+            int[] version = new int[2];
+            if (!EGL14.eglInitialize(mEglDisplay, version, 0, version, 1))
+            {
+                mEglDisplay = null;
+                Log.d(TAG, "unable to initialize EGL14 display");
+                return;
+            }
+
+            mCreatedEGLDisplay = true;
+        }
+
+        int[] configSpec = new int[]
+                {
+                        EGL14.EGL_RENDERABLE_TYPE, EGL14.EGL_OPENGL_ES2_BIT,
+                        EGL14.EGL_SURFACE_TYPE, EGL14.EGL_PBUFFER_BIT,
+                        EGL14.EGL_NONE
+                };
+        EGLConfig[] configs = new EGLConfig[1];
+        int[] num_config = new int[1];
+        EGL14.eglChooseConfig(mEglDisplay, configSpec, 0, configs, 0, 1, num_config, 0);
+        int[] contextAttribsES2 = new int[]
+                {
+                        EGL14.EGL_CONTEXT_CLIENT_VERSION, 2,
+                        EGL14.EGL_NONE
+                };
+        int[] contextAttribsES31 = new int[]
+                {
+                        EGLExt.EGL_CONTEXT_MAJOR_VERSION_KHR, 3,
+                        EGLExt.EGL_CONTEXT_MINOR_VERSION_KHR, 1,
+                        EGL14.EGL_NONE
+                };
+        mEglContext = EGL14.eglCreateContext(mEglDisplay, configs[0], shareContext, majorver[0]==3 ? contextAttribsES31 : contextAttribsES2, 0);
+
+        if (EGL14.eglQueryString(mEglDisplay, EGL14.EGL_EXTENSIONS).contains("EGL_KHR_surfaceless_context"))
+        {
+            mEglSurface = EGL14.EGL_NO_SURFACE;
+        }
+        else
+        {
+            int[] pbufferAttribs = new int[]
+                    {
+                            EGL14.EGL_NONE
+                    };
+            mEglSurface = EGL14.eglCreatePbufferSurface(mEglDisplay, configs[0], pbufferAttribs, 0);
+        }
+    }
+
+    private void saveContext()
+    {
+        mSavedDisplay = EGL14.eglGetCurrentDisplay();
+        mSavedContext = EGL14.eglGetCurrentContext();
+        mSavedSurfaceDraw = EGL14.eglGetCurrentSurface(EGL14.EGL_DRAW);
+        mSavedSurfaceRead = EGL14.eglGetCurrentSurface(EGL14.EGL_READ);
+    }
+
+    private void makeCurrent()
+    {
+        EGL14.eglMakeCurrent(mEglDisplay, mEglSurface, mEglSurface, mEglContext);
+    }
+
+    private void restoreContext()
+    {
+        EGL14.eglMakeCurrent(mSavedDisplay, mSavedSurfaceDraw, mSavedSurfaceRead, mSavedContext);
+    }
+
+    private void initSurfaceTexture()
+    {
+        int[] textures = new int[1];
+        GLES20.glGenTextures(1, textures, 0);
+        mTextureID = textures[0];
+        if (mTextureID <= 0)
+        {
+            Log.e(TAG, "mTextureID <= 0");
+            release();
+            return;
+        }
+        mSurfaceTexture = new android.graphics.SurfaceTexture(mTextureID);
+        mSurfaceTexture.setOnFrameAvailableListener(this);
+        mSurface = new android.view.Surface(mSurfaceTexture);
+
+        int[] glInt = new int[1];
+
+        GLES20.glGenFramebuffers(1,glInt,0);
+        mFBO = glInt[0];
+        if (mFBO <= 0)
+        {
+            Log.e(TAG, "mFBO <= 0");
+            release();
+            return;
+        }
+
+        // Special shaders for blit of movie texture.
+        mBlitVertexShaderID = createShader(GLES20.GL_VERTEX_SHADER, mBlitVextexShader);
+        if (mBlitVertexShaderID == 0)
+        {
+            Log.e(TAG, "mBlitVertexShaderID == 0");
+            release();
+            return;
+        }
+        int mBlitFragmentShaderID = createShader(GLES20.GL_FRAGMENT_SHADER,
+                mSwizzlePixels ? mBlitFragmentShaderBGRA : mBlitFragmentShaderRGBA);
+        if (mBlitFragmentShaderID == 0)
+        {
+            Log.e(TAG, "mBlitFragmentShaderID == 0");
+            release();
+            return;
+        }
+        mProgram = GLES20.glCreateProgram();
+        if (mProgram <= 0)
+        {
+            Log.e(TAG, "mProgram <= 0");
+            release();
+            return;
+        }
+        GLES20.glAttachShader(mProgram, mBlitVertexShaderID);
+        GLES20.glAttachShader(mProgram, mBlitFragmentShaderID);
+        GLES20.glLinkProgram(mProgram);
+        int[] linkStatus = new int[1];
+        GLES20.glGetProgramiv(mProgram, GLES20.GL_LINK_STATUS, linkStatus, 0);
+        if (linkStatus[0] != GLES20.GL_TRUE)
+        {
+            Log.e(TAG, "Could not link program: ");
+            Log.e(TAG, GLES20.glGetProgramInfoLog(mProgram));
+            GLES20.glDeleteProgram(mProgram);
+            mProgram = 0;
+            release();
+            return;
+        }
+        mPositionAttrib = GLES20.glGetAttribLocation(mProgram, "Position");
+        mTexCoordsAttrib = GLES20.glGetAttribLocation(mProgram, "TexCoords");
+        mTextureUniform = GLES20.glGetUniformLocation(mProgram, "VideoTexture");
+
+        GLES20.glGenBuffers(1,glInt,0);
+        mBlitBuffer = glInt[0];
+        if (mBlitBuffer <= 0)
+        {
+            Log.e(TAG, "mBlitBuffer <= 0");
+            release();
+            return;
+        }
+
+        // Create blit mesh.
+        mTriangleVertices = java.nio.ByteBuffer.allocateDirect(
+                        mTriangleVerticesData.length * FLOAT_SIZE_BYTES)
+                .order(java.nio.ByteOrder.nativeOrder()).asFloatBuffer();
+        mTriangleVerticesDirty = true;
+
+        // Set up GL state
+        if (mUseOwnContext)
+        {
+            GLES20.glDisable(GLES20.GL_BLEND);
+            GLES20.glDisable(GLES20.GL_CULL_FACE);
+            GLES20.glDisable(GLES20.GL_SCISSOR_TEST);
+            GLES20.glDisable(GLES20.GL_STENCIL_TEST);
+            GLES20.glDisable(GLES20.GL_DEPTH_TEST);
+            GLES20.glDisable(GLES20.GL_DITHER);
+            GLES20.glColorMask(true,true,true,true);
+        }
+    }
+
+    private void UpdateVertexData()
+    {
+        if (!mTriangleVerticesDirty || mBlitBuffer <= 0)
+        {
+            return;
+        }
+
+        // fill it in
+        mTriangleVertices.position(0);
+        mTriangleVertices.put(mTriangleVerticesData).position(0);
+
+        // save VBO state
+        int[] glInt = new int[1];
+        GLES20.glGetIntegerv(GLES20.GL_ARRAY_BUFFER_BINDING, glInt, 0);
+        int previousVBO = glInt[0];
+
+        GLES20.glBindBuffer(GLES20.GL_ARRAY_BUFFER, mBlitBuffer);
+        GLES20.glBufferData(GLES20.GL_ARRAY_BUFFER,
+                mTriangleVerticesData.length*FLOAT_SIZE_BYTES,
+                mTriangleVertices, GLES20.GL_STATIC_DRAW);
+
+        // restore VBO state
+        GLES20.glBindBuffer(GLES20.GL_ARRAY_BUFFER, previousVBO);
+
+        mTriangleVerticesDirty = false;
+    }
+
+    public boolean isValid()
+    {
+        return mSurfaceTexture != null;
+    }
+
+    private int createShader(int shaderType, String source)
+    {
+        int shader = GLES20.glCreateShader(shaderType);
+        if (shader != 0)
+        {
+            GLES20.glShaderSource(shader, source);
+            GLES20.glCompileShader(shader);
+            int[] compiled = new int[1];
+            GLES20.glGetShaderiv(shader, GLES20.GL_COMPILE_STATUS, compiled, 0);
+            if (compiled[0] == 0)
+            {
+                Log.e(TAG, "Could not compile shader " + shaderType + ":");
+                Log.e(TAG, GLES20.glGetShaderInfoLog(shader));
+                GLES20.glDeleteShader(shader);
+                shader = 0;
+            }
+        }
+        return shader;
+    }
+
+    public void onFrameAvailable(android.graphics.SurfaceTexture st)
+    {
+        synchronized(this)
+        {
+            mFrameAvailable = true;
+        }
+    }
+
+    public android.graphics.SurfaceTexture getSurfaceTexture()
+    {
+        return mSurfaceTexture;
+    }
+
+    public android.view.Surface getSurface()
+    {
+        return mSurface;
+    }
+
+    public int getExternalTextureId()
+    {
+        return mTextureID;
+    }
+
+    // NOTE: Synchronized with updateFrameData to prevent frame
+    // updates while the surface may need to get reallocated.
+    public void setSize(int width, int height)
+    {
+        synchronized(this)
+        {
+            if (width != mTextureWidth ||
+                    height != mTextureHeight)
+            {
+                mTextureWidth = width;
+                mTextureHeight = height;
+                mFrameData = null;
+                mTextureSizeChanged = true;
+            }
+        }
+    }
+
+    public boolean resolutionChanged()
+    {
+        boolean changed;
+        synchronized(this)
+        {
+            changed = mTextureSizeChanged;
+            mTextureSizeChanged = false;
+        }
+        return changed;
+    }
+
+    private static final int FLOAT_SIZE_BYTES = 4;
+    private static final int TRIANGLE_VERTICES_DATA_STRIDE_BYTES = 4 * FLOAT_SIZE_BYTES;
+    private static final int TRIANGLE_VERTICES_DATA_POS_OFFSET = 0;
+    private static final int TRIANGLE_VERTICES_DATA_UV_OFFSET = 2;
+    private float[] mTriangleVerticesData = {
+            // X, Y, U, V
+            -1.0f, -1.0f, 0.f, 0.f,
+            1.0f, -1.0f, 1.f, 0.f,
+            -1.0f, 1.0f, 0.f, 1.f,
+            1.0f, 1.0f, 1.f, 1.f,
+    };
+
+    private java.nio.FloatBuffer mTriangleVertices;
+
+    private final String mBlitVextexShader =
+            "attribute vec2 Position;\n" +
+                    "attribute vec2 TexCoords;\n" +
+                    "varying vec2 TexCoord;\n" +
+                    "void main() {\n" +
+                    "	TexCoord = TexCoords;\n" +
+                    "	gl_Position = vec4(Position, 0.0, 1.0);\n" +
+                    "}\n";
+
+    // NOTE: We read the fragment as BGRA so that in the end, when
+    // we glReadPixels out of the FBO, we get them in that order
+    // and avoid having to swizzle the pixels in the CPU.
+    private final String mBlitFragmentShaderBGRA =
+            "#extension GL_OES_EGL_image_external : require\n" +
+                    "uniform samplerExternalOES VideoTexture;\n" +
+                    "varying highp vec2 TexCoord;\n" +
+                    "void main()\n" +
+                    "{\n" +
+                    "	gl_FragColor = texture2D(VideoTexture, TexCoord).bgra;\n" +
+                    "}\n";
+    private final String mBlitFragmentShaderRGBA =
+            "#extension GL_OES_EGL_image_external : require\n" +
+                    "uniform samplerExternalOES VideoTexture;\n" +
+                    "varying highp vec2 TexCoord;\n" +
+                    "void main()\n" +
+                    "{\n" +
+                    "	gl_FragColor = texture2D(VideoTexture, TexCoord).rgba;\n" +
+                    "}\n";
+
+    private int mProgram;
+    private int mPositionAttrib;
+    private int mTexCoordsAttrib;
+    private int mBlitBuffer;
+    private int mTextureUniform;
+
+    public synchronized FrameData updateFrameData(UpdateParam param)
+    {
+        if (param.destTexture != 0)
+        {
+            FrameData frameData = new FrameData();
+            frameData.binaryBuffer = null;
+            frameData.frameReady = copyFrameTexture(param.destTexture, null);
+            return frameData;
+        }
+        else
+        {
+            if (null == mFrameData && mTextureWidth > 0 && mTextureHeight > 0)
+            {
+                mFrameData = java.nio.ByteBuffer.allocateDirect(mTextureWidth*mTextureHeight*4);
+            }
+            // Copy surface texture to frame data.
+            if (!copyFrameTexture(0, mFrameData))
+            {
+                return null;
+            }
+
+            FrameData frameData = new FrameData();
+            frameData.binaryBuffer = mFrameData;
+            frameData.frameReady = true;
+            return frameData;
+        }
+    }
+
+    public native void nativeClearCachedAttributeState(int PositionAttrib, int TexCoordsAttrib);
+
+    // Copy the surface texture to another texture, or to raw data.
+    // Note: copying to raw data creates a temporary FBO texture.
+    private boolean copyFrameTexture(int destTexture, java.nio.Buffer destData)
+    {
+        if (!mFrameAvailable)
+        {
+            // We only return fresh data when we generate it. At other
+            // time we return nothing to indicate that there was nothing
+            // new to return. The media player deals with this by keeping
+            // the last frame around and using that for rendering.
+            return false;
+        }
+
+        mFrameAvailable = false;
+
+        if (null == mSurfaceTexture)
+        {
+            // Can't update if there's no surface to update into.
+            return false;
+        }
+
+        int[] glInt = new int[1];
+        boolean[] glBool = new boolean[1];
+
+        // Either use own context or save states
+        boolean previousBlend=false, previousCullFace=false, previousScissorTest=false, previousStencilTest=false, previousDepthTest=false, previousDither=false;
+        int previousFBO=0, previousVBO=0, previousMinFilter=0, previousMagFilter=0;
+        int[] previousViewport = new int[4];
+        if (mUseOwnContext)
+        {
+            // Received reports of these not being preserved when changing contexts
+            GLES20.glActiveTexture(GLES20.GL_TEXTURE0);
+            GLES20.glGetTexParameteriv(GLES20.GL_TEXTURE_2D, GLES20.GL_TEXTURE_MIN_FILTER, glInt, 0);
+            previousMinFilter = glInt[0];
+            GLES20.glGetTexParameteriv(GLES20.GL_TEXTURE_2D, GLES20.GL_TEXTURE_MAG_FILTER, glInt, 0);
+            previousMagFilter = glInt[0];
+
+            saveContext();
+            makeCurrent();
+        }
+        else
+        {
+            // Clear gl errors as they can creep in from the UE4 renderer.
+            GLES20.glGetError();
+
+            previousBlend = GLES20.glIsEnabled(GLES20.GL_BLEND);
+            previousCullFace = GLES20.glIsEnabled(GLES20.GL_CULL_FACE);
+            previousScissorTest = GLES20.glIsEnabled(GLES20.GL_SCISSOR_TEST);
+            previousStencilTest = GLES20.glIsEnabled(GLES20.GL_STENCIL_TEST);
+            previousDepthTest = GLES20.glIsEnabled(GLES20.GL_DEPTH_TEST);
+            previousDither = GLES20.glIsEnabled(GLES20.GL_DITHER);
+            GLES20.glGetIntegerv(GLES20.GL_FRAMEBUFFER_BINDING, glInt, 0);
+            previousFBO = glInt[0];
+            GLES20.glGetIntegerv(GLES20.GL_ARRAY_BUFFER_BINDING, glInt, 0);
+            previousVBO = glInt[0];
+            GLES20.glGetIntegerv(GLES20.GL_VIEWPORT, previousViewport, 0);
+
+            GLES20.glActiveTexture(GLES20.GL_TEXTURE0);
+            GLES20.glGetTexParameteriv(GLES20.GL_TEXTURE_2D, GLES20.GL_TEXTURE_MIN_FILTER, glInt, 0);
+            previousMinFilter = glInt[0];
+            GLES20.glGetTexParameteriv(GLES20.GL_TEXTURE_2D, GLES20.GL_TEXTURE_MAG_FILTER, glInt, 0);
+            previousMagFilter = glInt[0];
+
+            glVerify("save state");
+        }
+
+        // Get the latest video texture frame.
+        mSurfaceTexture.updateTexImage();
+
+        mSurfaceTexture.getTransformMatrix(mTransformMatrix);
+
+        float UMin = mTransformMatrix[12];
+        float UMax = UMin + mTransformMatrix[0];
+        float VMin = mTransformMatrix[13];
+        float VMax = VMin + mTransformMatrix[5];
+
+        if (mTriangleVerticesData[2] != UMin ||
+                mTriangleVerticesData[6] != UMax ||
+                mTriangleVerticesData[11] != VMin ||
+                mTriangleVerticesData[3] != VMax)
+        {
+            //GameActivity.Log.debug("Matrix:");
+            //GameActivity.Log.debug(mTransformMatrix[0] + " " + mTransformMatrix[1] + " " + mTransformMatrix[2] + " " + mTransformMatrix[3]);
+            //GameActivity.Log.debug(mTransformMatrix[4] + " " + mTransformMatrix[5] + " " + mTransformMatrix[6] + " " + mTransformMatrix[7]);
+            //GameActivity.Log.debug(mTransformMatrix[8] + " " + mTransformMatrix[9] + " " + mTransformMatrix[10] + " " + mTransformMatrix[11]);
+            //GameActivity.Log.debug(mTransformMatrix[12] + " " + mTransformMatrix[13] + " " + mTransformMatrix[14] + " " + mTransformMatrix[15]);
+            mTriangleVerticesData[ 2] = mTriangleVerticesData[10] = UMin;
+            mTriangleVerticesData[ 6] = mTriangleVerticesData[14] = UMax;
+            mTriangleVerticesData[11] = mTriangleVerticesData[15] = VMin;
+            mTriangleVerticesData[ 3] = mTriangleVerticesData[ 7] = VMax;
+            mTriangleVerticesDirty = true;
+            //GameActivity.Log.debug("U = " + mTriangleVerticesData[2] + ", " + mTriangleVerticesData[6]);
+            //GameActivity.Log.debug("V = " + mTriangleVerticesData[11] + ", " + mTriangleVerticesData[3]);
+        }
+
+        if (null != destData)
+        {
+            // Rewind data so that we can write to it.
+            destData.position(0);
+        }
+
+        if (!mUseOwnContext)
+        {
+            GLES20.glDisable(GLES20.GL_BLEND);
+            GLES20.glDisable(GLES20.GL_CULL_FACE);
+            GLES20.glDisable(GLES20.GL_SCISSOR_TEST);
+            GLES20.glDisable(GLES20.GL_STENCIL_TEST);
+            GLES20.glDisable(GLES20.GL_DEPTH_TEST);
+            GLES20.glDisable(GLES20.GL_DITHER);
+            GLES20.glColorMask(true,true,true,true);
+
+            glVerify("reset state");
+        }
+
+        GLES20.glViewport(0, 0, mTextureWidth, mTextureHeight);
+
+        glVerify("set viewport");
+
+        // Set-up FBO target texture..
+        int FBOTextureID = 0;
+        if (null != destData)
+        {
+            // Create temporary FBO for data copy.
+            GLES20.glGenTextures(1,glInt,0);
+            FBOTextureID = glInt[0];
+        }
+        else
+        {
+            // Use the given texture as the FBO.
+            FBOTextureID = destTexture;
+        }
+        // Set the FBO to draw into the texture one-to-one.
+        GLES20.glBindTexture(GLES20.GL_TEXTURE_2D, FBOTextureID);
+        GLES20.glTexParameteri(GLES20.GL_TEXTURE_2D,
+                GLES20.GL_TEXTURE_MIN_FILTER, GLES20.GL_NEAREST);
+        GLES20.glTexParameteri(GLES20.GL_TEXTURE_2D,
+                GLES20.GL_TEXTURE_MAG_FILTER, GLES20.GL_NEAREST);
+        GLES20.glTexParameteri(GLES20.GL_TEXTURE_2D,
+                GLES20.GL_TEXTURE_WRAP_S, GLES20.GL_CLAMP_TO_EDGE);
+        GLES20.glTexParameteri(GLES20.GL_TEXTURE_2D,
+                GLES20.GL_TEXTURE_WRAP_T, GLES20.GL_CLAMP_TO_EDGE);
+        // Create the temp FBO data if needed.
+        if (null != destData)
+        {
+            //int w = 1<<(32-Integer.numberOfLeadingZeros(mTextureWidth-1));
+            //int h = 1<<(32-Integer.numberOfLeadingZeros(mTextureHeight-1));
+            GLES20.glTexImage2D(GLES20.GL_TEXTURE_2D, 0,
+                    GLES20.GL_RGBA,
+                    mTextureWidth, mTextureHeight,
+                    0, GLES20.GL_RGBA, GLES20.GL_UNSIGNED_BYTE, null);
+        }
+
+        glVerify("set-up FBO texture");
+
+        // Set to render to the FBO.
+        GLES20.glBindFramebuffer(GLES20.GL_FRAMEBUFFER, mFBO);
+
+        GLES20.glFramebufferTexture2D(
+                GLES20.GL_FRAMEBUFFER,
+                GLES20.GL_COLOR_ATTACHMENT0,
+                GLES20.GL_TEXTURE_2D, FBOTextureID, 0);
+
+        // check status
+        int status = GLES20.glCheckFramebufferStatus(GLES20.GL_FRAMEBUFFER);
+        if (status != GLES20.GL_FRAMEBUFFER_COMPLETE)
+        {
+            Log.e(TAG, "Failed to complete framebuffer attachment ("+status+")");
+        }
+
+        // The special shaders to render from the video texture.
+        GLES20.glUseProgram(mProgram);
+
+        // Set the mesh that renders the video texture.
+        UpdateVertexData();
+        GLES20.glBindBuffer(GLES20.GL_ARRAY_BUFFER, mBlitBuffer);
+        GLES20.glEnableVertexAttribArray(mPositionAttrib);
+        GLES20.glVertexAttribPointer(mPositionAttrib, 2, GLES20.GL_FLOAT, false,
+                TRIANGLE_VERTICES_DATA_STRIDE_BYTES, 0);
+        GLES20.glEnableVertexAttribArray(mTexCoordsAttrib);
+        GLES20.glVertexAttribPointer(mTexCoordsAttrib, 2, GLES20.GL_FLOAT, false,
+                TRIANGLE_VERTICES_DATA_STRIDE_BYTES,
+                TRIANGLE_VERTICES_DATA_UV_OFFSET*FLOAT_SIZE_BYTES);
+
+        glVerify("setup movie texture read");
+
+        GLES20.glClear( GLES20.GL_COLOR_BUFFER_BIT);
+
+        // connect 'VideoTexture' to video source texture (mTextureID).
+        // mTextureID is bound to GL_TEXTURE_EXTERNAL_OES in updateTexImage
+        GLES20.glUniform1i(mTextureUniform, 0);
+        GLES20.glActiveTexture(GLES20.GL_TEXTURE0);
+        GLES20.glBindTexture(GL_TEXTURE_EXTERNAL_OES, mTextureID);
+
+        // Draw the video texture mesh.
+        GLES20.glDrawArrays(GLES20.GL_TRIANGLE_STRIP, 0, 4);
+
+        GLES20.glFlush();
+
+        // Read the FBO texture pixels into raw data.
+        if (null != destData)
+        {
+            GLES20.glReadPixels(
+                    0, 0, mTextureWidth, mTextureHeight,
+                    GLES20.GL_RGBA, GLES20.GL_UNSIGNED_BYTE,
+                    destData);
+        }
+
+        glVerify("draw & read movie texture");
+
+        // Restore state and cleanup.
+        if (mUseOwnContext)
+        {
+            GLES20.glFramebufferTexture2D(
+                    GLES20.GL_FRAMEBUFFER,
+                    GLES20.GL_COLOR_ATTACHMENT0,
+                    GLES20.GL_TEXTURE_2D, 0, 0);
+
+            if (null != destData && FBOTextureID > 0)
+            {
+                glInt[0] = FBOTextureID;
+                GLES20.glDeleteTextures(1, glInt, 0);
+            }
+
+            restoreContext();
+
+            // Restore previous texture filtering
+            GLES20.glTexParameteri(GLES20.GL_TEXTURE_2D, GLES20.GL_TEXTURE_MIN_FILTER, previousMinFilter);
+            GLES20.glTexParameteri(GLES20.GL_TEXTURE_2D, GLES20.GL_TEXTURE_MAG_FILTER, previousMagFilter);
+        }
+        else
+        {
+            GLES20.glBindFramebuffer(GLES20.GL_FRAMEBUFFER, previousFBO);
+            if (null != destData && FBOTextureID > 0)
+            {
+                glInt[0] = FBOTextureID;
+                GLES20.glDeleteTextures(1, glInt, 0);
+            }
+            GLES20.glBindBuffer(GLES20.GL_ARRAY_BUFFER, previousVBO);
+
+            GLES20.glViewport(previousViewport[0], previousViewport[1],	previousViewport[2], previousViewport[3]);
+            if (previousBlend) GLES20.glEnable(GLES20.GL_BLEND);
+            if (previousCullFace) GLES20.glEnable(GLES20.GL_CULL_FACE);
+            if (previousScissorTest) GLES20.glEnable(GLES20.GL_SCISSOR_TEST);
+            if (previousStencilTest) GLES20.glEnable(GLES20.GL_STENCIL_TEST);
+            if (previousDepthTest) GLES20.glEnable(GLES20.GL_DEPTH_TEST);
+            if (previousDither) GLES20.glEnable(GLES20.GL_DITHER);
+
+            GLES20.glTexParameteri(GLES20.GL_TEXTURE_2D, GLES20.GL_TEXTURE_MIN_FILTER, previousMinFilter);
+            GLES20.glTexParameteri(GLES20.GL_TEXTURE_2D, GLES20.GL_TEXTURE_MAG_FILTER, previousMagFilter);
+
+            // invalidate cached state in RHI
+            GLES20.glDisableVertexAttribArray(mPositionAttrib);
+            GLES20.glDisableVertexAttribArray(mTexCoordsAttrib);
+            nativeClearCachedAttributeState(mPositionAttrib, mTexCoordsAttrib);
+        }
+
+        return true;
+    }
+
+    private void showGlError(String op, int error)
+    {
+        switch (error)
+        {
+            case GLES20.GL_INVALID_ENUM:						Log.e(TAG, "MediaPlayer$BitmapRenderer: " + op + ": glGetError GL_INVALID_ENUM");  break;
+            case GLES20.GL_INVALID_OPERATION:					Log.e(TAG, "MediaPlayer$BitmapRenderer: " + op + ": glGetError GL_INVALID_OPERATION");  break;
+            case GLES20.GL_INVALID_FRAMEBUFFER_OPERATION:		Log.e(TAG, "MediaPlayer$BitmapRenderer: " + op + ": glGetError GL_INVALID_FRAMEBUFFER_OPERATION");  break;
+            case GLES20.GL_INVALID_VALUE:						Log.e(TAG, "MediaPlayer$BitmapRenderer: " + op + ": glGetError GL_INVALID_VALUE");  break;
+            case GLES20.GL_FRAMEBUFFER_INCOMPLETE_ATTACHMENT:	Log.e(TAG, "MediaPlayer$BitmapRenderer: " + op + ": glGetError GL_FRAMEBUFFER_INCOMPLETE_ATTACHMENT");  break;
+            case GLES20.GL_FRAMEBUFFER_INCOMPLETE_DIMENSIONS:	Log.e(TAG, "MediaPlayer$BitmapRenderer: " + op + ": glGetError GL_FRAMEBUFFER_INCOMPLETE_DIMENSIONS");  break;
+            case GLES20.GL_FRAMEBUFFER_UNSUPPORTED:				Log.e(TAG, "MediaPlayer$BitmapRenderer: " + op + ": glGetError GL_FRAMEBUFFER_UNSUPPORTED");  break;
+            case GLES20.GL_OUT_OF_MEMORY:						Log.e(TAG, "MediaPlayer$BitmapRenderer: " + op + ": glGetError GL_OUT_OF_MEMORY");  break;
+            default:											Log.e(TAG, "MediaPlayer$BitmapRenderer: " + op + ": glGetError " + error);
+        }
+    }
+
+    private void glVerify(String op)
+    {
+        int error;
+        while ((error = GLES20.glGetError()) != GLES20.GL_NO_ERROR)
+        {
+            showGlError(op, error);
+            throw new RuntimeException(op + ": glGetError " + error);
+        }
+    }
+
+    private void glWarn(String op)
+    {
+        int error;
+        while ((error = GLES20.glGetError()) != GLES20.GL_NO_ERROR)
+        {
+            showGlError(op, error);
+        }
+    }
+
+    public void release()
+    {
+        if (null != mSurface)
+        {
+            mSurface.release();
+            mSurface = null;
+        }
+        if (null != mSurfaceTexture)
+        {
+            mSurfaceTexture.release();
+            mSurfaceTexture = null;
+        }
+        int[] glInt = new int[1];
+        if (mBlitBuffer > 0)
+        {
+            glInt[0] = mBlitBuffer;
+            GLES20.glDeleteBuffers(1,glInt,0);
+            mBlitBuffer = -1;
+        }
+        if (mProgram > 0)
+        {
+            GLES20.glDeleteProgram(mProgram);
+            mProgram = -1;
+        }
+        if (mBlitVertexShaderID > 0)
+        {
+            GLES20.glDeleteShader(mBlitVertexShaderID);
+            mBlitVertexShaderID = -1;
+        }
+        if (mBlitFragmentShaderID > 0)
+        {
+            GLES20.glDeleteShader(mBlitFragmentShaderID);
+            mBlitFragmentShaderID = -1;
+        }
+        if (mFBO > 0)
+        {
+            glInt[0] = mFBO;
+            GLES20.glDeleteFramebuffers(1,glInt,0);
+            mFBO = -1;
+        }
+        if (mTextureID > 0)
+        {
+            glInt[0] = mTextureID;
+            GLES20.glDeleteTextures(1,glInt,0);
+            mTextureID = -1;
+        }
+        if (mEglSurface != EGL14.EGL_NO_SURFACE)
+        {
+            EGL14.eglDestroySurface(mEglDisplay, mEglSurface);
+            mEglSurface = EGL14.EGL_NO_SURFACE;
+        }
+        if (mEglContext != EGL14.EGL_NO_CONTEXT)
+        {
+            EGL14.eglDestroyContext(mEglDisplay, mEglContext);
+            mEglContext = EGL14.EGL_NO_CONTEXT;
+        }
+        if (mCreatedEGLDisplay)
+        {
+            EGL14.eglTerminate(mEglDisplay);
+            mEglDisplay = EGL14.EGL_NO_DISPLAY;
+            mCreatedEGLDisplay = false;
+        }
+    }
+}
\ No newline at end of file
diff --git a/Engine/Build/Android/Java/src/com/epicgames/unreal/MediaPlayer14.java b/Engine/Build/Android/Java/src/com/epicgames/unreal/MediaPlayer14.java
--- a/Engine/Build/Android/Java/src/com/epicgames/unreal/MediaPlayer14.java
+++ b/Engine/Build/Android/Java/src/com/epicgames/unreal/MediaPlayer14.java
@@ -4,9 +4,8 @@ package com.epicgames.unreal;
 
 import android.content.res.AssetFileDescriptor;
 import android.content.res.AssetManager;
-import android.graphics.Rect;
 import android.opengl.*;
-import android.os.Build;
+
 import java.io.File;
 import java.io.IOException;
 import java.io.RandomAccessFile;
@@ -14,10 +13,8 @@ import android.media.MediaDataSource;
 import android.media.MediaExtractor;
 import android.media.MediaFormat;
 import android.media.MediaPlayer;
-import android.media.MediaPlayer.TrackInfo;
+
 import java.util.ArrayList;
-import java.util.Random;
-import android.util.Log;
 import java.net.URL;
 import java.net.HttpURLConnection;
 
@@ -591,11 +588,11 @@ public class MediaPlayer14
 		{
 			return mOESTextureRenderer.getExternalTextureId();
 		}
-		if (null != mBitmapRenderer)
+		if (!usingBitmapRendererImageReader && (null != mBitmapRenderer) && (mBitmapRenderer instanceof BitmapRendererLegacy))
 		{
-			return mBitmapRenderer.getExternalTextureId();
+			return ((BitmapRendererLegacy)mBitmapRenderer).getExternalTextureId();
 		}
-		return -1;
+		return 0; // 0 is the only value never used in GLES - the red book
 	}
 
 	public void prepare() throws IOException, IllegalStateException
@@ -747,11 +744,31 @@ public class MediaPlayer14
 
 	// ======================================================================================
 
+	private boolean usingBitmapRendererImageReader = false;
+	private static final int ImageReaderQueueLength = 16;
+	private BitmapRenderer.UpdateParam updateParam = new BitmapRenderer.UpdateParam();
+
+	public void setUsingBitmapRendererImageReader(boolean enabled)
+	{
+		usingBitmapRendererImageReader = enabled;
+	}
+
 	private boolean CreateBitmapRenderer()
 	{
 		releaseBitmapRenderer();
 
-		mBitmapRenderer = new BitmapRenderer(SwizzlePixels, VulkanRenderer);
+		GameActivity.Log.debug("CREATE BITMAPRENDERER, VIDEO SIZE: " + getVideoWidth() + " x " + getVideoHeight() + ", BRIR: " + usingBitmapRendererImageReader + ", Vulkan: " + VulkanRenderer);
+		
+		if (usingBitmapRendererImageReader)
+		{
+			mBitmapRenderer = new BitmapRendererImageReader(SwizzlePixels, VulkanRenderer,
+					getVideoWidth(), getVideoHeight(), ImageReaderQueueLength);
+		}
+		else
+		{
+			mBitmapRenderer = new BitmapRendererLegacy(SwizzlePixels, VulkanRenderer);
+		}
+
 		if (!mBitmapRenderer.isValid())
 		{
 			mBitmapRenderer = null;
@@ -800,19 +817,43 @@ public class MediaPlayer14
 			{
 				GameActivity.Log.warn("initBitmapRenderer failed to alloc mBitmapRenderer ");
 				reset();
-			  }
+			}
+		}
+	}
+
+	public BitmapRendererImageReader.FrameData getVideoLastMediaFrameData()
+	{
+		if (usingBitmapRendererImageReader)
+		{
+			initBitmapRenderer();
+			if (mBitmapRenderer != null)
+			{
+				WaitOnBitmapRender = true;
+				BitmapRendererImageReader.FrameData data = mBitmapRenderer.updateFrameData(updateParam);
+				WaitOnBitmapRender = false;
+				return data;
+			}
 		}
+		return null;
 	}
 
 	public java.nio.Buffer getVideoLastFrameData()
 	{
+		if (usingBitmapRendererImageReader)
+		{
+			return null;
+		}
+		
 		initBitmapRenderer();
 		if (null != mBitmapRenderer)
 		{
 			WaitOnBitmapRender = true;
-			java.nio.Buffer data = mBitmapRenderer.updateFrameData();
+
+			updateParam.destTexture = 0;
+			BitmapRendererImageReader.FrameData data = mBitmapRenderer.updateFrameData(updateParam);
+
 			WaitOnBitmapRender = false;
-			return data;
+			return data.binaryBuffer;
 		}
 		else
 		{
@@ -823,13 +864,21 @@ public class MediaPlayer14
 	public boolean getVideoLastFrame(int destTexture)
 	{
 //		GameActivity.Log.debug("getVideoLastFrame: " + destTexture);
+		if (usingBitmapRendererImageReader)
+		{
+			return false;
+		}
+
 		initBitmapRenderer();
 		if (null != mBitmapRenderer)
 		{
 			WaitOnBitmapRender = true;
-			boolean result = mBitmapRenderer.updateFrameData(destTexture);
+
+			updateParam.destTexture = destTexture;
+			BitmapRendererImageReader.FrameData data = mBitmapRenderer.updateFrameData(updateParam);
+
 			WaitOnBitmapRender = false;
-			return result;
+			return data.frameReady;
 		}
 		else
 		{
@@ -837,833 +886,6 @@ public class MediaPlayer14
 		}
 	}
 
-	/*
-		All this internal surface view does is manage the
-		offscreen bitmap that the media player decoding can
-		render into for eventual extraction to the UE4 buffers.
-	*/
-	class BitmapRenderer
-		implements android.graphics.SurfaceTexture.OnFrameAvailableListener
-	{
-		private java.nio.Buffer mFrameData = null;
-		private int mLastFramePosition = -1;
-		private android.graphics.SurfaceTexture mSurfaceTexture = null;
-		private int mTextureWidth = -1;
-		private int mTextureHeight = -1;
-		private android.view.Surface mSurface = null;
-		private boolean mFrameAvailable = false;
-		private int mTextureID = -1;
-		private int mFBO = -1;
-		private int mBlitVertexShaderID = -1;
-		private int mBlitFragmentShaderID = -1;
-		private float[] mTransformMatrix = new float[16];
-		private boolean mTriangleVerticesDirty = true;
-		private boolean mTextureSizeChanged = true;
-		private boolean mUseOwnContext = true;
-		private boolean mVulkanRenderer = false;
-		private boolean mSwizzlePixels = false;
-
-		private int GL_TEXTURE_EXTERNAL_OES = 0x8D65;
-
-		private EGLDisplay mEglDisplay;
-		private EGLContext mEglContext;
-		private EGLSurface mEglSurface;
-
-		private EGLDisplay mSavedDisplay;
-		private EGLContext mSavedContext;
-		private EGLSurface mSavedSurfaceDraw;
-		private EGLSurface mSavedSurfaceRead;
-
-		private boolean mCreatedEGLDisplay = false;
-
-		public BitmapRenderer(boolean swizzlePixels, boolean isVulkan)
-		{
-			mSwizzlePixels = swizzlePixels;
-			mVulkanRenderer = isVulkan;
-
-			mEglSurface = EGL14.EGL_NO_SURFACE;
-			mEglContext = EGL14.EGL_NO_CONTEXT;
-			mUseOwnContext = true;
-
-			if (mVulkanRenderer)
-			{
-				mSwizzlePixels = true;
-			}
-			else
-			{
-				String RendererString = GLES20.glGetString(GLES20.GL_RENDERER);
-
-				// Do not use shared context if Adreno before 400 or on older Android than Marshmallow
-				if (RendererString.contains("Adreno (TM) "))
-				{
-					java.util.regex.Pattern pattern = java.util.regex.Pattern.compile("Adreno \\(TM\\) (\\d*)");
-					java.util.regex.Matcher matcher = pattern.matcher(RendererString);
-					if (matcher.find() && matcher.groupCount() >= 1)
-					{
-						int AdrenoVersion = Integer.parseInt(matcher.group(1));
-						if (AdrenoVersion < 400 || android.os.Build.VERSION.SDK_INT < 22)
-						{
-							GameActivity.Log.debug("MediaPlayer14: disabled shared GL context on " + RendererString);
-							mUseOwnContext = false;
-						}
-					}
-				}
-			}
-
-			if (mUseOwnContext)
-			{
-				initContext();
-				saveContext();
-				makeCurrent();
-				initSurfaceTexture();
-				restoreContext();
-			}
-			else
-			{
-				initSurfaceTexture();
-			}
-		}
-
-		private void initContext()
-		{
-			mEglDisplay = EGL14.EGL_NO_DISPLAY;
-			EGLContext shareContext = EGL14.EGL_NO_CONTEXT;
-
-			int majorver[] = new int[] { 0 };
-			int minorver[] = new int[] { 0 };
-			if (!mVulkanRenderer)
-			{
-				mEglDisplay = EGL14.eglGetCurrentDisplay();
-				shareContext = EGL14.eglGetCurrentContext();
-
-				if (android.os.Build.VERSION.SDK_INT >= 18 &&
-					EGL14.eglQueryContext(mEglDisplay, shareContext, EGLExt.EGL_CONTEXT_MAJOR_VERSION_KHR, majorver, 0) &&
-					EGL14.eglQueryContext(mEglDisplay, shareContext, EGLExt.EGL_CONTEXT_MINOR_VERSION_KHR, minorver, 0))
-				{
-					GameActivity.Log.debug("MediaPlayer14: Existing GL context is version " + majorver[0] + "." + minorver[0]);
-				}
-				else
-				// on some devices eg Galaxy S6, the above fails but we do get EGL14.EGL_CONTEXT_CLIENT_VERSION=3
-				if (EGL14.eglQueryContext(mEglDisplay, shareContext, EGL14.EGL_CONTEXT_CLIENT_VERSION, majorver, 0))
-				{					
-					GameActivity.Log.debug("MediaPlayer14: Existing GL context is version " + majorver[0]);
-				}
-				else
-				{
-					GameActivity.Log.debug("MediaPlayer14: Existing GL context version not detected");		
-				}
-			}
-			else
-			{
-				mEglDisplay = EGL14.eglGetDisplay(EGL14.EGL_DEFAULT_DISPLAY);
-				if (mEglDisplay == EGL14.EGL_NO_DISPLAY)
-				{
-					GameActivity.Log.error("unable to get EGL14 display");
-					return;
-				}
-				int[] version = new int[2];
-				if (!EGL14.eglInitialize(mEglDisplay, version, 0, version, 1))
-				{
-					mEglDisplay = null;
-					GameActivity.Log.error("unable to initialize EGL14 display");
-					return;
-				}				
-				
-				mCreatedEGLDisplay = true;
-			}
-
-			int[] configSpec = new int[]
-			{
-				EGL14.EGL_RENDERABLE_TYPE, EGL14.EGL_OPENGL_ES2_BIT,
-				EGL14.EGL_SURFACE_TYPE, EGL14.EGL_PBUFFER_BIT,
-				EGL14.EGL_NONE
-			};
-			EGLConfig[] configs = new EGLConfig[1];
-			int[] num_config = new int[1];
-			EGL14.eglChooseConfig(mEglDisplay, configSpec, 0, configs, 0, 1, num_config, 0);
-			int[] contextAttribsES2 = new int[]
-			{
-				EGL14.EGL_CONTEXT_CLIENT_VERSION, 2,
-				EGL14.EGL_NONE
-			};
-			int[] contextAttribsES31 = new int[]
-			{
-				EGLExt.EGL_CONTEXT_MAJOR_VERSION_KHR, 3,
-				EGLExt.EGL_CONTEXT_MINOR_VERSION_KHR, 1,
-				EGL14.EGL_NONE
-			};
-			mEglContext = EGL14.eglCreateContext(mEglDisplay, configs[0], shareContext, majorver[0]==3 ? contextAttribsES31 : contextAttribsES2, 0);
-
-			if (EGL14.eglQueryString(mEglDisplay, EGL14.EGL_EXTENSIONS).contains("EGL_KHR_surfaceless_context"))
-			{
-				mEglSurface = EGL14.EGL_NO_SURFACE;
-			}
-			else
-			{
-				int[] pbufferAttribs = new int[]
-				{
-					EGL14.EGL_NONE
-				};
-				mEglSurface = EGL14.eglCreatePbufferSurface(mEglDisplay, configs[0], pbufferAttribs, 0);
-			}
-		}
-
-		private void saveContext()
-		{
-			mSavedDisplay = EGL14.eglGetCurrentDisplay();
-			mSavedContext = EGL14.eglGetCurrentContext();
-			mSavedSurfaceDraw = EGL14.eglGetCurrentSurface(EGL14.EGL_DRAW);
-			mSavedSurfaceRead = EGL14.eglGetCurrentSurface(EGL14.EGL_READ);
-		}
-
-		private void makeCurrent()
-		{
-			EGL14.eglMakeCurrent(mEglDisplay, mEglSurface, mEglSurface, mEglContext);
-		}
-
-		private void restoreContext()
-		{
-			EGL14.eglMakeCurrent(mSavedDisplay, mSavedSurfaceDraw, mSavedSurfaceRead, mSavedContext);
-		}
-
-		private void initSurfaceTexture()
-		{
-			int[] textures = new int[1];
-			GLES20.glGenTextures(1, textures, 0);
-			mTextureID = textures[0];
-			if (mTextureID <= 0)
-			{
-				GameActivity.Log.error("mTextureID <= 0");
-				release();
-				return;
-			}
-			mSurfaceTexture = new android.graphics.SurfaceTexture(mTextureID);
-			mSurfaceTexture.setOnFrameAvailableListener(this);
-			mSurface = new android.view.Surface(mSurfaceTexture);
-
-			int[] glInt = new int[1];
-
-			GLES20.glGenFramebuffers(1,glInt,0);
-			mFBO = glInt[0];
-			if (mFBO <= 0)
-			{
-				GameActivity.Log.error("mFBO <= 0");
-				release();
-				return;
-			}
-
-			// Special shaders for blit of movie texture.
-			mBlitVertexShaderID = createShader(GLES20.GL_VERTEX_SHADER, mBlitVextexShader);
-			if (mBlitVertexShaderID == 0)
-			{
-				GameActivity.Log.error("mBlitVertexShaderID == 0");
-				release();
-				return;
-			}
-			int mBlitFragmentShaderID = createShader(GLES20.GL_FRAGMENT_SHADER,
-				mSwizzlePixels ? mBlitFragmentShaderBGRA : mBlitFragmentShaderRGBA);
-			if (mBlitFragmentShaderID == 0)
-			{
-				GameActivity.Log.error("mBlitFragmentShaderID == 0");
-				release();
-				return;
-			}
-			mProgram = GLES20.glCreateProgram();
-			if (mProgram <= 0)
-			{
-				GameActivity.Log.error("mProgram <= 0");
-				release();
-				return;
-			}
-			GLES20.glAttachShader(mProgram, mBlitVertexShaderID);
-			GLES20.glAttachShader(mProgram, mBlitFragmentShaderID);
-			GLES20.glLinkProgram(mProgram);
-			int[] linkStatus = new int[1];
-			GLES20.glGetProgramiv(mProgram, GLES20.GL_LINK_STATUS, linkStatus, 0);
-			if (linkStatus[0] != GLES20.GL_TRUE)
-			{
-				GameActivity.Log.error("Could not link program: ");
-				GameActivity.Log.error(GLES20.glGetProgramInfoLog(mProgram));
-				GLES20.glDeleteProgram(mProgram);
-				mProgram = 0;
-				release();
-				return;
-			}
-			mPositionAttrib = GLES20.glGetAttribLocation(mProgram, "Position");
-			mTexCoordsAttrib = GLES20.glGetAttribLocation(mProgram, "TexCoords");
-			mTextureUniform = GLES20.glGetUniformLocation(mProgram, "VideoTexture");
-
-			GLES20.glGenBuffers(1,glInt,0);
-			mBlitBuffer = glInt[0];
-			if (mBlitBuffer <= 0)
-			{
-				GameActivity.Log.error("mBlitBuffer <= 0");
-				release();
-				return;
-			}
-
-			// Create blit mesh.
-			mTriangleVertices = java.nio.ByteBuffer.allocateDirect(
-				mTriangleVerticesData.length * FLOAT_SIZE_BYTES)
-					.order(java.nio.ByteOrder.nativeOrder()).asFloatBuffer();
-			mTriangleVerticesDirty = true;
-
-			// Set up GL state
-			if (mUseOwnContext)
-			{
-				GLES20.glDisable(GLES20.GL_BLEND);
-				GLES20.glDisable(GLES20.GL_CULL_FACE);
-				GLES20.glDisable(GLES20.GL_SCISSOR_TEST);
-				GLES20.glDisable(GLES20.GL_STENCIL_TEST);
-				GLES20.glDisable(GLES20.GL_DEPTH_TEST);
-				GLES20.glDisable(GLES20.GL_DITHER);
-				GLES20.glColorMask(true,true,true,true);
-			}
-		}
-		
-		private void UpdateVertexData()
-		{
-			if (!mTriangleVerticesDirty || mBlitBuffer <= 0)
-			{
-				return;
-			}
-
-			// fill it in
-			mTriangleVertices.position(0);
-			mTriangleVertices.put(mTriangleVerticesData).position(0);
-
-			// save VBO state
-			int[] glInt = new int[1];
-			GLES20.glGetIntegerv(GLES20.GL_ARRAY_BUFFER_BINDING, glInt, 0);
-			int previousVBO = glInt[0];
-			
-			GLES20.glBindBuffer(GLES20.GL_ARRAY_BUFFER, mBlitBuffer);
-			GLES20.glBufferData(GLES20.GL_ARRAY_BUFFER,
-				mTriangleVerticesData.length*FLOAT_SIZE_BYTES,
-				mTriangleVertices, GLES20.GL_STATIC_DRAW);
-
-			// restore VBO state
-			GLES20.glBindBuffer(GLES20.GL_ARRAY_BUFFER, previousVBO);
-			
-			mTriangleVerticesDirty = false;
-		}
-
-		public boolean isValid()
-		{
-			return mSurfaceTexture != null;
-		}
-
-		private int createShader(int shaderType, String source)
-		{
-			int shader = GLES20.glCreateShader(shaderType);
-			if (shader != 0)
-			{
-				GLES20.glShaderSource(shader, source);
-				GLES20.glCompileShader(shader);
-				int[] compiled = new int[1];
-				GLES20.glGetShaderiv(shader, GLES20.GL_COMPILE_STATUS, compiled, 0);
-				if (compiled[0] == 0)
-				{
-					GameActivity.Log.error("Could not compile shader " + shaderType + ":");
-					GameActivity.Log.error(GLES20.glGetShaderInfoLog(shader));
-					GLES20.glDeleteShader(shader);
-					shader = 0;
-				}
-			}
-			return shader;
-		}
-
-		public void onFrameAvailable(android.graphics.SurfaceTexture st)
-		{
-			synchronized(this)
-			{
-				mFrameAvailable = true;
-			}
-		}
-
-		public android.graphics.SurfaceTexture getSurfaceTexture()
-		{
-			return mSurfaceTexture;
-		}
-
-		public android.view.Surface getSurface()
-		{
-			return mSurface;
-		}
-
-		public int getExternalTextureId()
-		{
-			return mTextureID;
-		}
-
-		// NOTE: Synchronized with updateFrameData to prevent frame
-		// updates while the surface may need to get reallocated.
-		public void setSize(int width, int height)
-		{
-			synchronized(this)
-			{
-				if (width != mTextureWidth ||
-					height != mTextureHeight)
-				{
-					mTextureWidth = width;
-					mTextureHeight = height;
-					mFrameData = null;
-					mTextureSizeChanged = true;
-				}
-			}
-		}
-
-		public boolean resolutionChanged()
-		{
-			boolean changed;
-			synchronized(this)
-			{
-				changed = mTextureSizeChanged;
-				mTextureSizeChanged = false;
-			}
-			return changed;
-		}
-
-		private static final int FLOAT_SIZE_BYTES = 4;
-		private static final int TRIANGLE_VERTICES_DATA_STRIDE_BYTES = 4 * FLOAT_SIZE_BYTES;
-		private static final int TRIANGLE_VERTICES_DATA_POS_OFFSET = 0;
-		private static final int TRIANGLE_VERTICES_DATA_UV_OFFSET = 2;
-		private float[] mTriangleVerticesData = {
-			// X, Y, U, V
-			-1.0f, -1.0f, 0.f, 0.f,
-			1.0f, -1.0f, 1.f, 0.f,
-			-1.0f, 1.0f, 0.f, 1.f,
-			1.0f, 1.0f, 1.f, 1.f,
-			};
-
-		private java.nio.FloatBuffer mTriangleVertices;
-
-		private final String mBlitVextexShader =
-			"attribute vec2 Position;\n" +
-			"attribute vec2 TexCoords;\n" +
-			"varying vec2 TexCoord;\n" +
-			"void main() {\n" +
-			"	TexCoord = TexCoords;\n" +
-			"	gl_Position = vec4(Position, 0.0, 1.0);\n" +
-			"}\n";
-
-		// NOTE: We read the fragment as BGRA so that in the end, when
-		// we glReadPixels out of the FBO, we get them in that order
-		// and avoid having to swizzle the pixels in the CPU.
-		private final String mBlitFragmentShaderBGRA =
-			"#extension GL_OES_EGL_image_external : require\n" +
-			"uniform samplerExternalOES VideoTexture;\n" +
-			"varying highp vec2 TexCoord;\n" +
-			"void main()\n" +
-			"{\n" +
-			"	gl_FragColor = texture2D(VideoTexture, TexCoord).bgra;\n" +
-			"}\n";
-		private final String mBlitFragmentShaderRGBA =
-			"#extension GL_OES_EGL_image_external : require\n" +
-			"uniform samplerExternalOES VideoTexture;\n" +
-			"varying highp vec2 TexCoord;\n" +
-			"void main()\n" +
-			"{\n" +
-			"	gl_FragColor = texture2D(VideoTexture, TexCoord).rgba;\n" +
-			"}\n";
-
-		private int mProgram;
-		private int mPositionAttrib;
-		private int mTexCoordsAttrib;
-		private int mBlitBuffer;
-		private int mTextureUniform;
-
-		public java.nio.Buffer updateFrameData()
-		{
-			synchronized(this)
-			{
-				if (null == mFrameData && mTextureWidth > 0 && mTextureHeight > 0)
-				{
-					mFrameData = java.nio.ByteBuffer.allocateDirect(mTextureWidth*mTextureHeight*4);
-				}
-				// Copy surface texture to frame data.
-				if (!copyFrameTexture(0, mFrameData))
-				{
-					return null;
-				}
-			}
-			return mFrameData;
-		}
-
-		public boolean updateFrameData(int destTexture)
-		{
-			synchronized(this)
-			{
-				// Copy surface texture to destination texture.
-				if (!copyFrameTexture(destTexture, null))
-				{
-					return false;
-				}
-			}
-			return true;
-		}
-
-		// Copy the surface texture to another texture, or to raw data.
-		// Note: copying to raw data creates a temporary FBO texture.
-		private boolean copyFrameTexture(int destTexture, java.nio.Buffer destData)
-		{
-			if (!mFrameAvailable)
-			{
-				// We only return fresh data when we generate it. At other
-				// time we return nothing to indicate that there was nothing
-				// new to return. The media player deals with this by keeping
-				// the last frame around and using that for rendering.
-				return false;
-			}
-			mFrameAvailable = false;
-			int current_frame_position = getCurrentPosition();
-			mLastFramePosition = current_frame_position;
-			if (null == mSurfaceTexture)
-			{
-				// Can't update if there's no surface to update into.
-				return false;
-			}
-
-			int[] glInt = new int[1];
-			boolean[] glBool = new boolean[1];
-
-			// Either use own context or save states
-			boolean previousBlend=false, previousCullFace=false, previousScissorTest=false, previousStencilTest=false, previousDepthTest=false, previousDither=false;
-			int previousFBO=0, previousVBO=0, previousMinFilter=0, previousMagFilter=0;
-			int[] previousViewport = new int[4];
-			if (mUseOwnContext)
-			{
-				// Received reports of these not being preserved when changing contexts
-				GLES20.glActiveTexture(GLES20.GL_TEXTURE0);
-				GLES20.glGetTexParameteriv(GLES20.GL_TEXTURE_2D, GLES20.GL_TEXTURE_MIN_FILTER, glInt, 0);
-				previousMinFilter = glInt[0];
-				GLES20.glGetTexParameteriv(GLES20.GL_TEXTURE_2D, GLES20.GL_TEXTURE_MAG_FILTER, glInt, 0);
-				previousMagFilter = glInt[0];
-
-				saveContext();
-				makeCurrent();
-			}
-			else
-			{
-				// Clear gl errors as they can creep in from the UE4 renderer.
-				GLES20.glGetError();
-
-				previousBlend = GLES20.glIsEnabled(GLES20.GL_BLEND);
-				previousCullFace = GLES20.glIsEnabled(GLES20.GL_CULL_FACE);
-				previousScissorTest = GLES20.glIsEnabled(GLES20.GL_SCISSOR_TEST);
-				previousStencilTest = GLES20.glIsEnabled(GLES20.GL_STENCIL_TEST);
-				previousDepthTest = GLES20.glIsEnabled(GLES20.GL_DEPTH_TEST);
-				previousDither = GLES20.glIsEnabled(GLES20.GL_DITHER);
-				GLES20.glGetIntegerv(GLES20.GL_FRAMEBUFFER_BINDING, glInt, 0);
-				previousFBO = glInt[0];
-				GLES20.glGetIntegerv(GLES20.GL_ARRAY_BUFFER_BINDING, glInt, 0);
-				previousVBO = glInt[0];
-				GLES20.glGetIntegerv(GLES20.GL_VIEWPORT, previousViewport, 0);
-
-				GLES20.glActiveTexture(GLES20.GL_TEXTURE0);
-				GLES20.glGetTexParameteriv(GLES20.GL_TEXTURE_2D, GLES20.GL_TEXTURE_MIN_FILTER, glInt, 0);
-				previousMinFilter = glInt[0];
-				GLES20.glGetTexParameteriv(GLES20.GL_TEXTURE_2D, GLES20.GL_TEXTURE_MAG_FILTER, glInt, 0);
-				previousMagFilter = glInt[0];
-
-				glVerify("save state");
-			}
-
-			// Get the latest video texture frame.
-			mSurfaceTexture.updateTexImage();
-
-			mSurfaceTexture.getTransformMatrix(mTransformMatrix);
-				
-			float UMin = mTransformMatrix[12];
-			float UMax = UMin + mTransformMatrix[0];
-			float VMin = mTransformMatrix[13];
-			float VMax = VMin + mTransformMatrix[5];
-				
-			if (mTriangleVerticesData[2] != UMin ||
-				mTriangleVerticesData[6] != UMax ||
-				mTriangleVerticesData[11] != VMin ||
-				mTriangleVerticesData[3] != VMax)
-			{
-				//GameActivity.Log.debug("Matrix:");
-				//GameActivity.Log.debug(mTransformMatrix[0] + " " + mTransformMatrix[1] + " " + mTransformMatrix[2] + " " + mTransformMatrix[3]);
-				//GameActivity.Log.debug(mTransformMatrix[4] + " " + mTransformMatrix[5] + " " + mTransformMatrix[6] + " " + mTransformMatrix[7]);
-				//GameActivity.Log.debug(mTransformMatrix[8] + " " + mTransformMatrix[9] + " " + mTransformMatrix[10] + " " + mTransformMatrix[11]);
-				//GameActivity.Log.debug(mTransformMatrix[12] + " " + mTransformMatrix[13] + " " + mTransformMatrix[14] + " " + mTransformMatrix[15]);
-				mTriangleVerticesData[ 2] = mTriangleVerticesData[10] = UMin;
-				mTriangleVerticesData[ 6] = mTriangleVerticesData[14] = UMax;
-				mTriangleVerticesData[11] = mTriangleVerticesData[15] = VMin;
-				mTriangleVerticesData[ 3] = mTriangleVerticesData[ 7] = VMax;
-				mTriangleVerticesDirty = true;
-				//GameActivity.Log.debug("U = " + mTriangleVerticesData[2] + ", " + mTriangleVerticesData[6]);
-				//GameActivity.Log.debug("V = " + mTriangleVerticesData[11] + ", " + mTriangleVerticesData[3]);
-			}
-
-			if (null != destData)
-			{
-				// Rewind data so that we can write to it.
-				destData.position(0);
-			}
-
-			if (!mUseOwnContext)
-			{
-				GLES20.glDisable(GLES20.GL_BLEND);
-				GLES20.glDisable(GLES20.GL_CULL_FACE);
-				GLES20.glDisable(GLES20.GL_SCISSOR_TEST);
-				GLES20.glDisable(GLES20.GL_STENCIL_TEST);
-				GLES20.glDisable(GLES20.GL_DEPTH_TEST);
-				GLES20.glDisable(GLES20.GL_DITHER);
-				GLES20.glColorMask(true,true,true,true);
-
-				glVerify("reset state");
-			}
-
-			GLES20.glViewport(0, 0, mTextureWidth, mTextureHeight);
-
-			glVerify("set viewport");
-
-			// Set-up FBO target texture..
-			int FBOTextureID = 0;
-			if (null != destData)
-			{
-				// Create temporary FBO for data copy.
-				GLES20.glGenTextures(1,glInt,0);
-				FBOTextureID = glInt[0];
-			}
-			else
-			{
-				// Use the given texture as the FBO.
-				FBOTextureID = destTexture;
-			}
-			// Set the FBO to draw into the texture one-to-one.
-			GLES20.glBindTexture(GLES20.GL_TEXTURE_2D, FBOTextureID);
-			GLES20.glTexParameteri(GLES20.GL_TEXTURE_2D,
-				GLES20.GL_TEXTURE_MIN_FILTER, GLES20.GL_NEAREST);
-			GLES20.glTexParameteri(GLES20.GL_TEXTURE_2D,
-				GLES20.GL_TEXTURE_MAG_FILTER, GLES20.GL_NEAREST);
-			GLES20.glTexParameteri(GLES20.GL_TEXTURE_2D,
-				GLES20.GL_TEXTURE_WRAP_S, GLES20.GL_CLAMP_TO_EDGE);
-			GLES20.glTexParameteri(GLES20.GL_TEXTURE_2D,
-				GLES20.GL_TEXTURE_WRAP_T, GLES20.GL_CLAMP_TO_EDGE);
-			// Create the temp FBO data if needed.
-			if (null != destData)
-			{
-				//int w = 1<<(32-Integer.numberOfLeadingZeros(mTextureWidth-1));
-				//int h = 1<<(32-Integer.numberOfLeadingZeros(mTextureHeight-1));
-				GLES20.glTexImage2D(GLES20.GL_TEXTURE_2D, 0,
-					GLES20.GL_RGBA,
-					mTextureWidth, mTextureHeight,
-					0, GLES20.GL_RGBA, GLES20.GL_UNSIGNED_BYTE, null);
-			}
-
-			glVerify("set-up FBO texture");
-
-			// Set to render to the FBO.
-			GLES20.glBindFramebuffer(GLES20.GL_FRAMEBUFFER, mFBO);
-
-			GLES20.glFramebufferTexture2D(
-				GLES20.GL_FRAMEBUFFER,
-				GLES20.GL_COLOR_ATTACHMENT0,
-				GLES20.GL_TEXTURE_2D, FBOTextureID, 0);
-
-			// check status
-			int status = GLES20.glCheckFramebufferStatus(GLES20.GL_FRAMEBUFFER);
-			if (status != GLES20.GL_FRAMEBUFFER_COMPLETE)
-			{
-				GameActivity.Log.warn("Failed to complete framebuffer attachment ("+status+")");
-			}
-
-			// The special shaders to render from the video texture.
-			GLES20.glUseProgram(mProgram);
-
-			// Set the mesh that renders the video texture.
-			UpdateVertexData();
-			GLES20.glBindBuffer(GLES20.GL_ARRAY_BUFFER, mBlitBuffer);
-			GLES20.glEnableVertexAttribArray(mPositionAttrib);
-			GLES20.glVertexAttribPointer(mPositionAttrib, 2, GLES20.GL_FLOAT, false,
-				TRIANGLE_VERTICES_DATA_STRIDE_BYTES, 0);
-			GLES20.glEnableVertexAttribArray(mTexCoordsAttrib);
-			GLES20.glVertexAttribPointer(mTexCoordsAttrib, 2, GLES20.GL_FLOAT, false,
-				TRIANGLE_VERTICES_DATA_STRIDE_BYTES,
-				TRIANGLE_VERTICES_DATA_UV_OFFSET*FLOAT_SIZE_BYTES);
-
-			glVerify("setup movie texture read");
-
-			GLES20.glClear( GLES20.GL_COLOR_BUFFER_BIT);
-
-			// connect 'VideoTexture' to video source texture (mTextureID).
-			// mTextureID is bound to GL_TEXTURE_EXTERNAL_OES in updateTexImage
-			GLES20.glUniform1i(mTextureUniform, 0);
-			GLES20.glActiveTexture(GLES20.GL_TEXTURE0);
-			GLES20.glBindTexture(GL_TEXTURE_EXTERNAL_OES, mTextureID);
-
-			// Draw the video texture mesh.
-			GLES20.glDrawArrays(GLES20.GL_TRIANGLE_STRIP, 0, 4);
-
-			GLES20.glFlush();
-
-			// Read the FBO texture pixels into raw data.
-			if (null != destData)
-			{
-				GLES20.glReadPixels(
-					0, 0, mTextureWidth, mTextureHeight,
-					GLES20.GL_RGBA, GLES20.GL_UNSIGNED_BYTE,
-					destData);
-			}
-
-			glVerify("draw & read movie texture");
-
-			// Restore state and cleanup.
-			if (mUseOwnContext)
-			{
-				GLES20.glFramebufferTexture2D(
-					GLES20.GL_FRAMEBUFFER,
-					GLES20.GL_COLOR_ATTACHMENT0,
-					GLES20.GL_TEXTURE_2D, 0, 0);
-
-				if (null != destData && FBOTextureID > 0)
-				{
-					glInt[0] = FBOTextureID;
-					GLES20.glDeleteTextures(1, glInt, 0);
-				}
-
-				restoreContext();
-
-				// Restore previous texture filtering
-				GLES20.glTexParameteri(GLES20.GL_TEXTURE_2D, GLES20.GL_TEXTURE_MIN_FILTER, previousMinFilter);
-				GLES20.glTexParameteri(GLES20.GL_TEXTURE_2D, GLES20.GL_TEXTURE_MAG_FILTER, previousMagFilter);
-			}
-			else
-			{
-				GLES20.glBindFramebuffer(GLES20.GL_FRAMEBUFFER, previousFBO);
-				if (null != destData && FBOTextureID > 0)
-				{
-					glInt[0] = FBOTextureID;
-					GLES20.glDeleteTextures(1, glInt, 0);
-				}
-				GLES20.glBindBuffer(GLES20.GL_ARRAY_BUFFER, previousVBO);
-
-				GLES20.glViewport(previousViewport[0], previousViewport[1],	previousViewport[2], previousViewport[3]);
-				if (previousBlend) GLES20.glEnable(GLES20.GL_BLEND);
-				if (previousCullFace) GLES20.glEnable(GLES20.GL_CULL_FACE);
-				if (previousScissorTest) GLES20.glEnable(GLES20.GL_SCISSOR_TEST);
-				if (previousStencilTest) GLES20.glEnable(GLES20.GL_STENCIL_TEST);
-				if (previousDepthTest) GLES20.glEnable(GLES20.GL_DEPTH_TEST);
-				if (previousDither) GLES20.glEnable(GLES20.GL_DITHER);
-
-				GLES20.glTexParameteri(GLES20.GL_TEXTURE_2D, GLES20.GL_TEXTURE_MIN_FILTER, previousMinFilter);
-				GLES20.glTexParameteri(GLES20.GL_TEXTURE_2D, GLES20.GL_TEXTURE_MAG_FILTER, previousMagFilter);
-
-				// invalidate cached state in RHI
-				GLES20.glDisableVertexAttribArray(mPositionAttrib);
-				GLES20.glDisableVertexAttribArray(mTexCoordsAttrib);
-				nativeClearCachedAttributeState(mPositionAttrib, mTexCoordsAttrib);
-			}
-
-			return true;
-		}
-
-		private void showGlError(String op, int error)
-		{
-			switch (error)
-			{
-				case GLES20.GL_INVALID_ENUM:						GameActivity.Log.error("MediaPlayer$BitmapRenderer: " + op + ": glGetError GL_INVALID_ENUM");  break;
-				case GLES20.GL_INVALID_OPERATION:					GameActivity.Log.error("MediaPlayer$BitmapRenderer: " + op + ": glGetError GL_INVALID_OPERATION");  break;
-				case GLES20.GL_INVALID_FRAMEBUFFER_OPERATION:		GameActivity.Log.error("MediaPlayer$BitmapRenderer: " + op + ": glGetError GL_INVALID_FRAMEBUFFER_OPERATION");  break;
-				case GLES20.GL_INVALID_VALUE:						GameActivity.Log.error("MediaPlayer$BitmapRenderer: " + op + ": glGetError GL_INVALID_VALUE");  break;
-				case GLES20.GL_FRAMEBUFFER_INCOMPLETE_ATTACHMENT:	GameActivity.Log.error("MediaPlayer$BitmapRenderer: " + op + ": glGetError GL_FRAMEBUFFER_INCOMPLETE_ATTACHMENT");  break;
-				case GLES20.GL_FRAMEBUFFER_INCOMPLETE_DIMENSIONS:	GameActivity.Log.error("MediaPlayer$BitmapRenderer: " + op + ": glGetError GL_FRAMEBUFFER_INCOMPLETE_DIMENSIONS");  break;
-				case GLES20.GL_FRAMEBUFFER_UNSUPPORTED:				GameActivity.Log.error("MediaPlayer$BitmapRenderer: " + op + ": glGetError GL_FRAMEBUFFER_UNSUPPORTED");  break;
-				case GLES20.GL_OUT_OF_MEMORY:						GameActivity.Log.error("MediaPlayer$BitmapRenderer: " + op + ": glGetError GL_OUT_OF_MEMORY");  break;
-				default:											GameActivity.Log.error("MediaPlayer$BitmapRenderer: " + op + ": glGetError " + error);
-			}
-		}
-
-		private void glVerify(String op)
-		{
-			int error;
-			while ((error = GLES20.glGetError()) != GLES20.GL_NO_ERROR)
-			{
-				showGlError(op, error);
-				throw new RuntimeException(op + ": glGetError " + error);
-			}
-		}
-
-		private void glWarn(String op)
-		{
-			int error;
-			while ((error = GLES20.glGetError()) != GLES20.GL_NO_ERROR)
-			{
-				showGlError(op, error);
-			}
-		}
-
-		public void release()
-		{
-			if (null != mSurface)
-			{
-				mSurface.release();
-				mSurface = null;
-			}
-			if (null != mSurfaceTexture)
-			{
-				mSurfaceTexture.release();
-				mSurfaceTexture = null;
-			}
-			int[] glInt = new int[1];
-			if (mBlitBuffer > 0)
-			{
-				glInt[0] = mBlitBuffer;
-				GLES20.glDeleteBuffers(1,glInt,0);
-				mBlitBuffer = -1;
-			}
-			if (mProgram > 0)
-			{
-				GLES20.glDeleteProgram(mProgram);
-				mProgram = -1;
-			}
-			if (mBlitVertexShaderID > 0)
-			{
-				GLES20.glDeleteShader(mBlitVertexShaderID);
-				mBlitVertexShaderID = -1;
-			}
-			if (mBlitFragmentShaderID > 0)
-			{
-				GLES20.glDeleteShader(mBlitFragmentShaderID);
-				mBlitFragmentShaderID = -1;
-			}
-			if (mFBO > 0)
-			{
-				glInt[0] = mFBO;
-				GLES20.glDeleteFramebuffers(1,glInt,0);
-				mFBO = -1;
-			}
-			if (mTextureID > 0)
-			{
-				glInt[0] = mTextureID;
-				GLES20.glDeleteTextures(1,glInt,0);
-				mTextureID = -1;
-			}
-			if (mEglSurface != EGL14.EGL_NO_SURFACE)
-			{
-				EGL14.eglDestroySurface(mEglDisplay, mEglSurface);
-				mEglSurface = EGL14.EGL_NO_SURFACE;
-			}
-			if (mEglContext != EGL14.EGL_NO_CONTEXT)
-			{
-				EGL14.eglDestroyContext(mEglDisplay, mEglContext);
-				mEglContext = EGL14.EGL_NO_CONTEXT;
-			}
-			if (mCreatedEGLDisplay)
-			{
-				EGL14.eglTerminate(mEglDisplay);
-				mEglDisplay = EGL14.EGL_NO_DISPLAY;
-				mCreatedEGLDisplay = false;	
-			}
-		}
-	};
-
-	public native void nativeClearCachedAttributeState(int PositionAttrib, int TexCoordsAttrib);
-
 	// ======================================================================================
 
 	private boolean CreateOESTextureRenderer(int OESTextureId)
diff --git a/Engine/Plugins/Media/AndroidMedia/Source/AndroidMedia/Private/Player/AndroidMediaPlayer.cpp b/Engine/Plugins/Media/AndroidMedia/Source/AndroidMedia/Private/Player/AndroidMediaPlayer.cpp
--- a/Engine/Plugins/Media/AndroidMedia/Source/AndroidMedia/Private/Player/AndroidMediaPlayer.cpp
+++ b/Engine/Plugins/Media/AndroidMedia/Source/AndroidMedia/Private/Player/AndroidMediaPlayer.cpp
@@ -666,6 +666,51 @@ struct FRHICommandUpdateBufferMediaSample final : public FRHICommand<FRHICommand
 	}
 };
 
+static void DoUpdateMediaFrameDataMediaSampleExecute(TWeakPtr<FJavaAndroidMediaPlayer, ESPMode::ThreadSafe> JavaMediaPlayerPtr, TWeakPtr<FMediaSamples, ESPMode::ThreadSafe> SamplesPtr,
+	TSharedRef<FAndroidMediaTextureSample, ESPMode::ThreadSafe> VideoSample)
+{
+	auto PinnedJavaMediaPlayer = JavaMediaPlayerPtr.Pin();
+	auto PinnedSamples = SamplesPtr.Pin();
+
+	if (!PinnedJavaMediaPlayer.IsValid() || !PinnedSamples.IsValid())
+	{
+		return;
+	}
+
+	const int32 CurrentFramePosition = PinnedJavaMediaPlayer->GetCurrentPosition();
+	const FTimespan Time = FTimespan::FromMilliseconds(CurrentFramePosition);
+
+	jobject FrameDataGlobalRef = PinnedJavaMediaPlayer->GetVideoLastMediaFrameData();
+	if (FrameDataGlobalRef == nullptr)
+	{
+		return;
+	}
+
+	VideoSample->InitializeMediaFrameData(FrameDataGlobalRef, Time);
+
+	PinnedSamples->AddVideo(VideoSample);
+}
+
+struct FRHICommandUpdateMediaFrameDataMediaSample final : public FRHICommand<FRHICommandUpdateMediaFrameDataMediaSample>
+{
+	TWeakPtr<FJavaAndroidMediaPlayer, ESPMode::ThreadSafe> JavaMediaPlayerPtr;
+	TWeakPtr<FMediaSamples, ESPMode::ThreadSafe> SamplesPtr;
+	TSharedRef<FAndroidMediaTextureSample, ESPMode::ThreadSafe> VideoSample;
+
+	FORCEINLINE_DEBUGGABLE FRHICommandUpdateMediaFrameDataMediaSample(TWeakPtr<FJavaAndroidMediaPlayer, ESPMode::ThreadSafe> InJavaMediaPlayerPtr, TWeakPtr<FMediaSamples, ESPMode::ThreadSafe> InSamplesPtr,
+		TSharedRef<FAndroidMediaTextureSample, ESPMode::ThreadSafe> InVideoSample)
+		: JavaMediaPlayerPtr(InJavaMediaPlayerPtr)
+		, SamplesPtr(InSamplesPtr)
+		, VideoSample(InVideoSample)
+	{
+	}
+	void Execute(FRHICommandListBase& CmdList)
+	{
+		QUICK_SCOPE_CYCLE_COUNTER(STAT_FRHICommandUpdateTextureMediaSample_Execute);
+		DoUpdateMediaFrameDataMediaSampleExecute(JavaMediaPlayerPtr, SamplesPtr, VideoSample);
+	}
+};
+
 // ==============================================================
 
 void FAndroidMediaPlayer::TickFetch(FTimespan DeltaTime, FTimespan /*Timecode*/)
@@ -715,40 +760,22 @@ void FAndroidMediaPlayer::TickFetch(FTimespan DeltaTime, FTimespan /*Timecode*/)
 		}
 		Params = { JavaMediaPlayer, Samples, VideoSample, (int32)(VideoTrack.Dimensions.X * VideoTrack.Dimensions.Y * sizeof(int32)) };
 
-		ENQUEUE_RENDER_COMMAND(AndroidMediaPlayerWriteVideoSample)(
-			[Params](FRHICommandListImmediate& RHICmdList)
-			{
-				auto PinnedJavaMediaPlayer = Params.JavaMediaPlayerPtr.Pin();
-				auto PinnedSamples = Params.SamplesPtr.Pin();
-
-				if (!PinnedJavaMediaPlayer.IsValid() || !PinnedSamples.IsValid())
-				{
-					return;
-				}
-
-				const int32 CurrentFramePosition = PinnedJavaMediaPlayer->GetCurrentPosition();
-				const FTimespan Time = FTimespan::FromMilliseconds(CurrentFramePosition);
-
-				// write frame into buffer
-				void* Buffer = nullptr;
-				int64 BufferCount = 0;
-
-				if (!PinnedJavaMediaPlayer->GetVideoLastFrameData(Buffer, BufferCount))
+		if (!JavaMediaPlayer->IsUsingBitmapRendererImageReader())
+		{
+			ENQUEUE_RENDER_COMMAND(AndroidMediaPlayerWriteVideoSample)(
+				[Params](FRHICommandListImmediate& RHICmdList)
 				{
-					return;
-				}
-
-				if (BufferCount != Params.SampleCount)
+					DoUpdateBufferMediaSampleExecute(Params.JavaMediaPlayerPtr, Params.SamplesPtr, Params.VideoSample, Params.SampleCount);
+				});
+		}
+		else
+		{
+			ENQUEUE_RENDER_COMMAND(AndroidMediaPlayerWriteVideoSample)(
+				[Params](FRHICommandListImmediate& RHICmdList)
 				{
-					FPlatformMisc::LowLevelOutputDebugStringf(TEXT("FAndroidMediaPlayer::Fetch: Sample count mismatch (Buffer=%d, Available=%d"), Params.SampleCount, BufferCount);
-				}
-				check(Params.SampleCount <= BufferCount);
-
-				// must make a copy (buffer is owned by Java, not us!)
-				Params.VideoSample->InitializeBuffer(Buffer, Time, true);
-
-				PinnedSamples->AddVideo(Params.VideoSample);
-			});
+					DoUpdateMediaFrameDataMediaSampleExecute(Params.JavaMediaPlayerPtr, Params.SamplesPtr, Params.VideoSample);
+				});
+		}
 	}
 #if ANDROIDMEDIAPLAYER_USE_EXTERNALTEXTURE
 	else if (GSupportsImageExternal)
@@ -770,50 +797,61 @@ void FAndroidMediaPlayer::TickFetch(FTimespan DeltaTime, FTimespan /*Timecode*/)
 			
 		FWriteVideoSampleParams Params = { JavaMediaPlayer, PlayerGuid, VideoSample, Samples };
 
-		ENQUEUE_RENDER_COMMAND(AndroidMediaPlayerWriteVideoSample)(
-			[Params](FRHICommandListImmediate& RHICmdList)
-			{
-				ConditionalInitializeVideoTexture_RenderThread(Params.JavaMediaPlayerPtr, Params.PlayerGuid);
-				if (IsRunningRHIInSeparateThread())
-				{
-					new (RHICmdList.AllocCommand<FRHICommandUpdateExternalMediaSample>()) FRHICommandUpdateExternalMediaSample(Params.JavaMediaPlayerPtr, Params.PlayerGuid);
-					// wait for DoUpdateExternalMediaSampleExecute to complete.
-					RHICmdList.ImmediateFlush(EImmediateFlushType::FlushRHIThread); 
-				}
-				else
-				{
-					DoUpdateExternalMediaSampleExecute(Params.JavaMediaPlayerPtr, Params.PlayerGuid);
-				}
-
-				auto PinnedJavaMediaPlayer = Params.JavaMediaPlayerPtr.Pin();
-				auto PinnedSamples = Params.Samples.Pin();
-
-				if (!PinnedJavaMediaPlayer.IsValid() || ! PinnedSamples.IsValid())
-				{
-					return;
-				}
-
-				FTextureRHIRef VideoTexture = PinnedJavaMediaPlayer->GetVideoTexture();
-				if (VideoTexture != nullptr && !PinnedJavaMediaPlayer->IsVideoTextureValid())
+		if (!JavaMediaPlayer->IsUsingBitmapRendererImageReader())
+		{
+			ENQUEUE_RENDER_COMMAND(AndroidMediaPlayerWriteVideoSample)(
+				[Params](FRHICommandListImmediate& RHICmdList)
 				{
+					ConditionalInitializeVideoTexture_RenderThread(Params.JavaMediaPlayerPtr, Params.PlayerGuid);
+					if (IsRunningRHIInSeparateThread())
+					{
+						new (RHICmdList.AllocCommand<FRHICommandUpdateExternalMediaSample>()) FRHICommandUpdateExternalMediaSample(Params.JavaMediaPlayerPtr, Params.PlayerGuid);
+						// wait for DoUpdateExternalMediaSampleExecute to complete.
+						RHICmdList.ImmediateFlush(EImmediateFlushType::FlushRHIThread);
+					}
+					else
+					{
+						DoUpdateExternalMediaSampleExecute(Params.JavaMediaPlayerPtr, Params.PlayerGuid);
+					}
+
+					auto PinnedJavaMediaPlayer = Params.JavaMediaPlayerPtr.Pin();
+					auto PinnedSamples = Params.Samples.Pin();
+
+					if (!PinnedJavaMediaPlayer.IsValid() || !PinnedSamples.IsValid())
+					{
+						return;
+					}
+
+					FTextureRHIRef VideoTexture = PinnedJavaMediaPlayer->GetVideoTexture();
+					if (VideoTexture != nullptr && !PinnedJavaMediaPlayer->IsVideoTextureValid())
+					{
 #if ANDROIDMEDIAPLAYER_USE_NATIVELOGGING
-					FPlatformMisc::LowLevelOutputDebugStringf(TEXT("Fetch RT: Register Guid: %s"), *Params.PlayerGuid.ToString());
+						FPlatformMisc::LowLevelOutputDebugStringf(TEXT("Fetch RT: Register Guid: %s"), *Params.PlayerGuid.ToString());
 #endif
-					FSamplerStateInitializerRHI SamplerStateInitializer(SF_Bilinear, AM_Clamp, AM_Clamp, AM_Clamp);
-					FSamplerStateRHIRef SamplerStateRHI = RHICreateSamplerState(SamplerStateInitializer);
-					FExternalTextureRegistry::Get().RegisterExternalTexture(Params.PlayerGuid, VideoTexture, SamplerStateRHI, FLinearColor(PinnedJavaMediaPlayer->GetUScale(), 0.0f, 0.0f, PinnedJavaMediaPlayer->GetVScale()), FLinearColor(PinnedJavaMediaPlayer->GetUOffset(), PinnedJavaMediaPlayer->GetVOffset(), 0.0f, 0.0f));
+						FSamplerStateInitializerRHI SamplerStateInitializer(SF_Bilinear, AM_Clamp, AM_Clamp, AM_Clamp);
+						FSamplerStateRHIRef SamplerStateRHI = RHICreateSamplerState(SamplerStateInitializer);
+						FExternalTextureRegistry::Get().RegisterExternalTexture(Params.PlayerGuid, VideoTexture, SamplerStateRHI, FLinearColor(PinnedJavaMediaPlayer->GetUScale(), 0.0f, 0.0f, PinnedJavaMediaPlayer->GetVScale()), FLinearColor(PinnedJavaMediaPlayer->GetUOffset(), PinnedJavaMediaPlayer->GetVOffset(), 0.0f, 0.0f));
 
-					PinnedJavaMediaPlayer->SetVideoTextureValid(true);
-				}
+						PinnedJavaMediaPlayer->SetVideoTextureValid(true);
+					}
+
+					const int32 CurrentFramePosition = PinnedJavaMediaPlayer->GetCurrentPosition();
+					const FTimespan Time = FTimespan::FromMilliseconds(CurrentFramePosition);
 
-				const int32 CurrentFramePosition = PinnedJavaMediaPlayer->GetCurrentPosition();
-				const FTimespan Time = FTimespan::FromMilliseconds(CurrentFramePosition);
-				
-				Params.VideoSample->InitializeExternalTexture(Time);
-				// Add a dummy video sample
-				// (actual data flows "on the side" via external texture mechanism)
-				PinnedSamples->AddVideo(Params.VideoSample);
-		});
+					Params.VideoSample->InitializeExternalTexture(Time);
+					// Add a dummy video sample
+					// (actual data flows "on the side" via external texture mechanism)
+					PinnedSamples->AddVideo(Params.VideoSample);
+				});
+		}
+		else
+		{
+			ENQUEUE_RENDER_COMMAND(AndroidMediaPlayerWriteVideoSample)(
+				[Params](FRHICommandListImmediate& RHICmdList)
+				{
+					DoUpdateMediaFrameDataMediaSampleExecute(Params.JavaMediaPlayerPtr, Params.Samples, Params.VideoSample);
+				});
+		}
 	}
 #endif // ANDROIDMEDIAPLAYER_USE_EXTERNALTEXTURE
 	else
@@ -877,18 +915,36 @@ void FAndroidMediaPlayer::TickFetch(FTimespan DeltaTime, FTimespan /*Timecode*/)
 
 	FWriteVideoSampleParams Params = { JavaMediaPlayer, Samples, VideoSample, (int32)(VideoTrack.Dimensions.X * VideoTrack.Dimensions.Y * sizeof(int32)) };
 
-	ENQUEUE_RENDER_COMMAND(AndroidMediaPlayerWriteVideoSample)(
-		[Params](FRHICommandListImmediate& RHICmdList)
-		{
-			if (IsRunningRHIInSeparateThread())
+	if (!JavaMediaPlayer->IsUsingBitmapRendererImageReader())
+	{
+		ENQUEUE_RENDER_COMMAND(AndroidMediaPlayerWriteVideoSample)(
+			[Params](FRHICommandListImmediate& RHICmdList)
 			{
-				new (RHICmdList.AllocCommand<FRHICommandUpdateBufferMediaSample>()) FRHICommandUpdateBufferMediaSample(Params.JavaMediaPlayerPtr, Params.SamplesPtr, Params.VideoSample, Params.SampleCount);
-			}
-			else
+				if (IsRunningRHIInSeparateThread())
+				{
+					new (RHICmdList.AllocCommand<FRHICommandUpdateBufferMediaSample>()) FRHICommandUpdateBufferMediaSample(Params.JavaMediaPlayerPtr, Params.SamplesPtr, Params.VideoSample, Params.SampleCount);
+				}
+				else
+				{
+					DoUpdateBufferMediaSampleExecute(Params.JavaMediaPlayerPtr, Params.SamplesPtr, Params.VideoSample, Params.SampleCount);
+				}
+			});
+	}
+	else
+	{
+		ENQUEUE_RENDER_COMMAND(AndroidMediaPlayerWriteVideoSample)(
+			[Params](FRHICommandListImmediate& RHICmdList)
 			{
-				DoUpdateBufferMediaSampleExecute(Params.JavaMediaPlayerPtr, Params.SamplesPtr, Params.VideoSample, Params.SampleCount);
-			}
-		});
+				if (IsRunningRHIInSeparateThread())
+				{
+					new (RHICmdList.AllocCommand<FRHICommandUpdateHardwareBufferMediaSample>()) FRHICommandUpdateHardwareBufferMediaSample(Params.JavaMediaPlayerPtr, Params.SamplesPtr, Params.VideoSample);
+				}
+				else
+				{
+					DoUpdateMediaFrameDataMediaSampleExecute(Params.JavaMediaPlayerPtr, Params.SamplesPtr, Params.VideoSample);
+				}
+			});
+	}
 #endif //WITH_ENGINE
 }
 
diff --git a/Engine/Plugins/Media/AndroidMedia/Source/AndroidMedia/Private/Player/AndroidMediaTextureSample.h b/Engine/Plugins/Media/AndroidMedia/Source/AndroidMedia/Private/Player/AndroidMediaTextureSample.h
--- a/Engine/Plugins/Media/AndroidMedia/Source/AndroidMedia/Private/Player/AndroidMediaTextureSample.h
+++ b/Engine/Plugins/Media/AndroidMedia/Source/AndroidMedia/Private/Player/AndroidMediaTextureSample.h
@@ -17,12 +17,17 @@
 	#include "MediaObjectPool.h"
 #endif
 
+#include "IMediaTextureSampleConverter.h"
+#include "Android/AndroidJava.h"
+#include "Android/AndroidJavaMediaFrameData.h"
+
 
 /**
  * Texture sample generated by AndroidMedia player.
  */
 class FAndroidMediaTextureSample
 	: public IMediaTextureSample
+	, public IMediaTextureSampleConverter
 	, public IMediaPoolable
 {
 public:
@@ -46,6 +51,8 @@ public:
 		{
 			FMemory::Free(Buffer);
 		}
+
+		MediaFrameData.CleanUp();
 	}
 
 public:
@@ -166,6 +173,21 @@ public:
 		return Texture;
 	}
 
+	/**
+	 * Initialize the sample with a java frame data object.
+	 *
+	 * @param InFrameData holds the jobject of the frame data.
+	 * @param InTime The sample time (in the player's local clock).
+	 * @note This method must be called on the render thread.
+	 * @see InitializeBuffer, InitializeTexture
+	 */
+	void InitializeMediaFrameData(jobject InFrameData, FTimespan InTime)
+	{
+		MediaFrameData.CleanUp();
+		MediaFrameData.Set(InFrameData);
+		Time = InTime;
+	}
+
 	/**
 	* Set the sample Scale, Rotation, Offset.
 	*
@@ -261,8 +283,64 @@ public:
 		return Offset;
 	}
 
+	virtual IMediaTextureSampleConverter* GetMediaTextureSampleConverter() override
+	{
+		if (MediaFrameData)
+		{
+			return this;
+		}
+		
+		return nullptr;
+	}
+
+	//~ IMediaPoolable interface
+	virtual bool IsReadyForReuse() override
+	{
+		if (MediaFrameData)
+		{
+			return MediaFrameData.IsReadyToClean();
+		}
+		return true;
+	}
+
+	virtual void ShutdownPoolable() override
+	{
+		// Reuse the texture?
+		//Texture = nullptr;
+
+		if (BufferSize > 0)
+		{
+			FMemory::Free(Buffer);
+			BufferSize = 0;
+			Buffer = nullptr;
+		}
+
+		MediaFrameData.CleanUp();
+	}
+
 private:
 
+	//~ IMediaTextureSampleConverter overrides
+	virtual bool Convert(FRHICommandListImmediate& RHICmdList, FTextureRHIRef& InDstTexture, const FConversionHints& Hints) override
+	{
+		if (!MediaFrameData)
+		{
+			return false;
+		}
+
+		if (FAndroidMisc::ShouldUseVulkan())
+		{
+			return MediaFrameData.ExtractToTextureVulkan(RHICmdList, InDstTexture, this);
+		}
+		else
+		{
+			return MediaFrameData.ExtractToTextureOES(RHICmdList, InDstTexture, this);
+		}
+	}
+
+	/** The sample's java frame data from BitmapRendererImageReader. */
+	FAndroidJavaMediaFrameData MediaFrameData;
+
 	/** The sample's data buffer. */
 	void* Buffer;
 
diff --git a/Engine/Source/Runtime/ApplicationCore/ApplicationCore.Build.cs b/Engine/Source/Runtime/ApplicationCore/ApplicationCore.Build.cs
--- a/Engine/Source/Runtime/ApplicationCore/ApplicationCore.Build.cs
+++ b/Engine/Source/Runtime/ApplicationCore/ApplicationCore.Build.cs
@@ -88,7 +88,9 @@ public class ApplicationCore : ModuleRules
 		{
 			PrivateIncludePathModuleNames.AddRange(
 				new string[] {
-					"Launch"
+					"Launch",
+					"Media",
+					"VulkanRHI"
 				}
 			);
 		}
diff --git a/Engine/Source/Runtime/ApplicationCore/Private/Android/AndroidJavaMediaFrameData.cpp b/Engine/Source/Runtime/ApplicationCore/Private/Android/AndroidJavaMediaFrameData.cpp
new file mode 100644
--- /dev/null
+++ b/Engine/Source/Runtime/ApplicationCore/Private/Android/AndroidJavaMediaFrameData.cpp
@@ -0,0 +1,294 @@
+// Copyright Epic Games, Inc. All Rights Reserved.
+
+#include "Android/AndroidJavaMediaFrameData.h"
+#include "Android/AndroidJava.h"
+#include "Android/AndroidApplication.h"
+
+#include "GlobalShader.h"
+#include "MediaShaders.h"
+#include "ColorManagement/ColorSpace.h"
+#include "PipelineStateCache.h"
+#include "ShaderParameterUtils.h"
+#include "IVulkanDynamicRHI.h"
+#include "IOpenGLDynamicRHI.h"
+#include "RHIStaticStates.h"
+
+#include <android/hardware_buffer.h>
+#include <android/hardware_buffer_jni.h>
+
+#if USE_ANDROID_JNI
+
+jfieldID FAndroidJavaMediaFrameData::FrameData_HardwareBufferHandle = nullptr;
+jfieldID FAndroidJavaMediaFrameData::FrameData_UScale = nullptr;
+jfieldID FAndroidJavaMediaFrameData::FrameData_UOffset = nullptr;
+jfieldID FAndroidJavaMediaFrameData::FrameData_VScale = nullptr;
+jfieldID FAndroidJavaMediaFrameData::FrameData_VOffset = nullptr;
+jmethodID FAndroidJavaMediaFrameData::FrameData_ReleaseFN = nullptr;
+
+FAndroidJavaMediaFrameData::FAndroidJavaMediaFrameData()
+	: FrameDataGlobalRef(nullptr)
+{
+	if (FrameData_HardwareBufferHandle == nullptr)
+	{
+		JNIEnv* JEnv = FAndroidApplication::GetJavaEnv();
+		jclass FrameDataClass = FAndroidApplication::FindJavaClass("com/epicgames/unreal/BitmapRenderer$FrameData");
+		FrameData_HardwareBufferHandle = JEnv->GetFieldID(FrameDataClass, "hardwareBuffer", "Landroid/hardware/HardwareBuffer;");
+		FrameData_UScale = JEnv->GetFieldID(FrameDataClass, "UScale", "F");
+		FrameData_UOffset = JEnv->GetFieldID(FrameDataClass, "UOffset", "F");
+		FrameData_VScale = JEnv->GetFieldID(FrameDataClass, "VScale", "F");
+		FrameData_VOffset = JEnv->GetFieldID(FrameDataClass, "VOffset", "F");
+		FrameData_ReleaseFN = JEnv->GetMethodID(FrameDataClass, "release", "()V");
+	}
+}
+
+FAndroidJavaMediaFrameData::~FAndroidJavaMediaFrameData()
+{
+	CleanUp();
+}
+
+FAndroidJavaMediaFrameData& FAndroidJavaMediaFrameData::operator=(FAndroidJavaMediaFrameData& Other)
+{
+	Swap(FrameDataGlobalRef, Other.FrameDataGlobalRef);
+	return *this;
+}
+
+bool FAndroidJavaMediaFrameData::Set(jobject InFrameData)
+{
+	CleanUp();
+
+	if (InFrameData != nullptr)
+	{
+		JNIEnv* JEnv = FAndroidApplication::GetJavaEnv();
+
+		jobjectRefType RefType = JEnv->GetObjectRefType(InFrameData);
+		if (RefType == jobjectRefType::JNILocalRefType)
+		{
+			FrameDataGlobalRef = JEnv->NewGlobalRef(InFrameData);
+		}
+		else if (RefType == jobjectRefType::JNIGlobalRefType)
+		{
+			FrameDataGlobalRef = InFrameData;
+		}
+	}
+
+	return (FrameDataGlobalRef != nullptr);
+}
+
+jobject FAndroidJavaMediaFrameData::Extract(FVector2f& OutScale, FVector2f& OutOffset)
+{
+	if (FrameDataGlobalRef == nullptr)
+	{
+		return nullptr;
+	}
+
+	JNIEnv* JEnv = FAndroidApplication::GetJavaEnv();
+	OutScale.X = JEnv->GetFloatField(FrameDataGlobalRef, FrameData_UScale);
+	OutScale.Y = JEnv->GetFloatField(FrameDataGlobalRef, FrameData_VScale);
+	OutOffset.X = JEnv->GetFloatField(FrameDataGlobalRef, FrameData_UOffset);
+	OutOffset.Y = JEnv->GetFloatField(FrameDataGlobalRef, FrameData_VOffset);
+
+	jobject HardwareBufferObj = JEnv->GetObjectField(FrameDataGlobalRef, FrameData_HardwareBufferHandle);
+	check(HardwareBufferObj != nullptr);
+
+	return HardwareBufferObj;
+}
+
+void FAndroidJavaMediaFrameData::CleanUp()
+{
+	Fence = nullptr;
+
+	if (FrameDataGlobalRef != nullptr)
+	{
+		JNIEnv* JEnv = FAndroidApplication::GetJavaEnv();
+		JEnv->CallVoidMethod(FrameDataGlobalRef, FrameData_ReleaseFN);
+		JEnv->DeleteGlobalRef(FrameDataGlobalRef);
+		FrameDataGlobalRef = nullptr;
+	}
+}
+
+bool FAndroidJavaMediaFrameData::ExtractToTextureVulkan(FRHICommandListImmediate& RHICmdList, FTextureRHIRef& InDstTexture, IMediaTextureSample* TextureSample)
+{
+	if (!FrameDataGlobalRef)
+	{
+		return false;
+	}
+
+	check(IsInRenderingThread());
+
+	bool Converted = false;
+	FVector2f FrameDataScale;
+	FVector2f FrameDataOffset;
+
+	jobject HardwareBufferObj = Extract(FrameDataScale, FrameDataOffset);
+	AHardwareBuffer* HardwareBuffer = AHardwareBuffer_fromHardwareBuffer(FAndroidApplication::GetJavaEnv(), HardwareBufferObj);
+	if (HardwareBuffer != nullptr)
+	{
+		AHardwareBuffer_acquire(HardwareBuffer);
+
+		IVulkanDynamicRHI* RHI = GetIVulkanDynamicRHI();
+		FTextureRHIRef SrcTexture = RHI->RHICreateTexture2DFromAndroidHardwareBuffer(HardwareBuffer);
+		if (SrcTexture.IsValid())
+		{
+			FGraphicsPipelineStateInitializer GraphicsPSOInit;
+			FRHITexture* RenderTarget = InDstTexture;
+
+			Fence = RHICreateGPUFence(TEXT("MediaFrameDataToTextureVulkan"));
+
+			RHICmdList.Transition(FRHITransitionInfo(InDstTexture, ERHIAccess::Unknown, ERHIAccess::RTV));
+
+			FRHIRenderPassInfo RPInfo(RenderTarget, ERenderTargetActions::DontLoad_Store);
+			RHICmdList.BeginRenderPass(RPInfo, TEXT("ConvertMedia_HardwareBuffer"));
+			{
+				RHICmdList.ApplyCachedRenderTargets(GraphicsPSOInit);
+
+				FIntPoint Dim = TextureSample->GetOutputDim();
+				RHICmdList.SetViewport(0, 0, 0.0f, (float)Dim.X, (float)Dim.Y, 1.0f);
+
+				GraphicsPSOInit.DepthStencilState = TStaticDepthStencilState<false, CF_Always>::GetRHI();
+				GraphicsPSOInit.RasterizerState = TStaticRasterizerState<>::GetRHI();
+				GraphicsPSOInit.BlendState = TStaticBlendStateWriteMask<CW_RGBA, CW_NONE, CW_NONE, CW_NONE, CW_NONE, CW_NONE, CW_NONE, CW_NONE>::GetRHI();
+				GraphicsPSOInit.PrimitiveType = PT_TriangleStrip;
+
+				// configure media shaders
+				FGlobalShaderMap* ShaderMap = GetGlobalShaderMap(GMaxRHIFeatureLevel);
+				TShaderMapRef<FMediaShadersVS> VertexShader(ShaderMap);
+
+				GraphicsPSOInit.BoundShaderState.VertexDeclarationRHI = GMediaVertexDeclaration.VertexDeclarationRHI;
+				GraphicsPSOInit.BoundShaderState.VertexShaderRHI = VertexShader.GetVertexShader();
+
+				FMatrix44f YUVMtx = TextureSample->GetSampleToRGBMatrix();
+				FMatrix44f ColorSpaceMtx;
+				//------------------------------------------------------------------------
+				//GetColorSpaceConversionMatrixForSample(this, ColorSpaceMtx);
+				{
+					const UE::Color::FColorSpace& Working = UE::Color::FColorSpace::GetWorking();
+
+					if (TextureSample->GetMediaTextureSampleColorConverter())
+					{
+						ColorSpaceMtx = FMatrix44f::Identity;
+					}
+					else
+					{
+						ColorSpaceMtx = UE::Color::Transpose<float>(UE::Color::FColorSpaceTransform(TextureSample->GetSourceColorSpace(), Working));
+					}
+
+					float NF = TextureSample->GetHDRNitsNormalizationFactor();
+					if (NF != 1.0f)
+					{
+						ColorSpaceMtx = ColorSpaceMtx.ApplyScale(NF);
+					}
+				}
+				//------------------------------------------------------------------------
+
+				TShaderMapRef<FVYUConvertPS> ConvertShader(ShaderMap);
+				GraphicsPSOInit.BoundShaderState.PixelShaderRHI = ConvertShader.GetPixelShader();
+				SetGraphicsPipelineState(RHICmdList, GraphicsPSOInit, 0);
+				SetShaderParametersLegacyPS(RHICmdList, ConvertShader, SrcTexture, Dim, YUVMtx, TextureSample->GetEncodingType(), ColorSpaceMtx);
+
+				// draw full size quad into render target
+				FBufferRHIRef VertexBuffer = CreateTempMediaVertexBuffer(RHICmdList);
+				RHICmdList.SetStreamSource(0, VertexBuffer, 0);
+				// set viewport to RT size
+				RHICmdList.SetViewport(0, 0, 0.0f, (float)Dim.X, (float)Dim.Y, 1.0f);
+
+				RHICmdList.DrawPrimitive(0, 2, 1);
+			}
+			RHICmdList.EndRenderPass();
+
+			RHICmdList.Transition(FRHITransitionInfo(RenderTarget, ERHIAccess::RTV, ERHIAccess::SRVGraphics));
+
+			RHICmdList.WriteGPUFence(Fence);
+
+			Converted = true;
+		}
+
+		AHardwareBuffer_release(HardwareBuffer);
+	}
+
+	FAndroidApplication::GetJavaEnv()->DeleteLocalRef(HardwareBufferObj);
+	return Converted;
+}
+
+bool FAndroidJavaMediaFrameData::ExtractToTextureOES(FRHICommandListImmediate& RHICmdList, FTextureRHIRef& InDstTexture, IMediaTextureSample* TextureSample)
+{
+	if (!FrameDataGlobalRef)
+	{
+		return false;
+	}
+
+	check(IsInRenderingThread());
+
+	bool Converted = false;
+	FVector2f FrameDataScale;
+	FVector2f FrameDataOffset;
+
+	jobject HardwareBufferObj = Extract(FrameDataScale, FrameDataOffset);
+	AHardwareBuffer* HardwareBuffer = AHardwareBuffer_fromHardwareBuffer(FAndroidApplication::GetJavaEnv(), HardwareBufferObj);
+	if (HardwareBuffer != nullptr)
+	{
+		AHardwareBuffer_acquire(HardwareBuffer);
+
+		IOpenGLDynamicRHI* RHI = GetIOpenGLDynamicRHI();
+		FTextureRHIRef SrcTexture = RHI->RHICreateTexture2DFromAndroidHardwareBuffer(RHICmdList, HardwareBuffer);
+		if (SrcTexture.IsValid())
+		{
+			FLinearColor Offset = { FrameDataOffset.X, FrameDataOffset.Y, 0.f, 0.f };
+			//FLinearColor ScaleRotation = { FrameDataScale.X, 0.f, 0.f, FrameDataScale.Y };
+			FLinearColor ScaleRotation = { 1.f, 0.f, 0.f, 1.f };
+
+			FGraphicsPipelineStateInitializer GraphicsPSOInit;
+			FRHITexture* RenderTarget = InDstTexture.GetReference();
+
+			RHICmdList.Transition(FRHITransitionInfo(InDstTexture, ERHIAccess::Unknown, ERHIAccess::RTV));
+
+			FRHIRenderPassInfo RPInfo(RenderTarget, ERenderTargetActions::DontLoad_Store);
+			RHICmdList.BeginRenderPass(RPInfo, TEXT("ConvertMedia_HardwareBuffer"));
+			{
+				RHICmdList.ApplyCachedRenderTargets(GraphicsPSOInit);
+
+				FIntPoint Dim = TextureSample->GetOutputDim();
+				RHICmdList.SetViewport(0, 0, 0.0f, (float)Dim.X, (float)Dim.Y, 1.0f);
+
+				GraphicsPSOInit.DepthStencilState = TStaticDepthStencilState<false, CF_Always>::GetRHI();
+				GraphicsPSOInit.RasterizerState = TStaticRasterizerState<>::GetRHI();
+				GraphicsPSOInit.BlendState = TStaticBlendStateWriteMask<CW_RGBA, CW_NONE, CW_NONE, CW_NONE, CW_NONE, CW_NONE, CW_NONE, CW_NONE>::GetRHI();
+				GraphicsPSOInit.PrimitiveType = PT_TriangleStrip;
+
+				// configure media shaders
+				auto ShaderMap = GetGlobalShaderMap(GMaxRHIFeatureLevel);
+				TShaderMapRef<FMediaShadersVS> VertexShader(ShaderMap);
+
+				GraphicsPSOInit.BoundShaderState.VertexDeclarationRHI = GMediaVertexDeclaration.VertexDeclarationRHI;
+				GraphicsPSOInit.BoundShaderState.VertexShaderRHI = VertexShader.GetVertexShader();
+
+				FSamplerStateInitializerRHI InSamplerStateInitializer(SF_Bilinear, AM_Clamp, AM_Clamp, AM_Clamp);
+				FSamplerStateRHIRef InSamplerState = RHICreateSamplerState(InSamplerStateInitializer);
+
+				TShaderMapRef<FReadTextureExternalPS> CopyShader(ShaderMap);
+				GraphicsPSOInit.BoundShaderState.PixelShaderRHI = CopyShader.GetPixelShader();
+				SetGraphicsPipelineState(RHICmdList, GraphicsPSOInit, 0);
+				SetShaderParametersLegacyPS(RHICmdList, CopyShader, SrcTexture, InSamplerState, ScaleRotation, Offset);
+
+				// draw full size quad into render target
+				FBufferRHIRef VertexBuffer = CreateTempMediaVertexBuffer(RHICmdList);
+				RHICmdList.SetStreamSource(0, VertexBuffer, 0);
+				// set viewport to RT size
+				RHICmdList.SetViewport(0, 0, 0.0f, (float)Dim.X, (float)Dim.Y, 1.0f);
+
+				RHICmdList.DrawPrimitive(0, 2, 1);
+			}
+			RHICmdList.EndRenderPass();
+
+			RHICmdList.Transition(FRHITransitionInfo(InDstTexture, ERHIAccess::RTV, ERHIAccess::SRVMask));
+
+			Converted = true;
+		}
+
+		AHardwareBuffer_release(HardwareBuffer);
+	}
+
+	FAndroidApplication::GetJavaEnv()->DeleteLocalRef(HardwareBufferObj);
+	return Converted;
+}
+
+#endif //USE_ANDROID_JNI
\ No newline at end of file
diff --git a/Engine/Source/Runtime/ApplicationCore/Private/Android/AndroidJavaMediaPlayer.cpp b/Engine/Source/Runtime/ApplicationCore/Private/Android/AndroidJavaMediaPlayer.cpp
--- a/Engine/Source/Runtime/ApplicationCore/Private/Android/AndroidJavaMediaPlayer.cpp
+++ b/Engine/Source/Runtime/ApplicationCore/Private/Android/AndroidJavaMediaPlayer.cpp
@@ -119,6 +119,13 @@ int32 FJavaAndroidMediaDataSource::ReadAt(int64 Position, uint8* Buffer, int32 C
 /* JavaAndroidMediaPlayer
 *****************************************************************************/
 
+bool AndroidEnableMediaPlayerBRIR = true;
+static FAutoConsoleVariableRef CVarAndroidEnableMediaPlayerBRIR(
+	TEXT("Android.EnableMediaPlayerBRIR"),
+	AndroidEnableMediaPlayerBRIR,
+	TEXT("Determines how the Android Media Player handles video frame data. True for using ImageReader/HardwareBuffer. False for copying data through the CPU memory. (Default: true)"),
+	ECVF_Default);
+
 FJavaAndroidMediaPlayer::FJavaAndroidMediaPlayer(bool swizzlePixels, bool vulkanRenderer, bool needTrackInfo)
 	: FJavaClassObject(GetClassName(), "(ZZZ)V", swizzlePixels, vulkanRenderer, needTrackInfo)
 	, GetDurationMethod(GetClassMethod("getDuration", "()I"))
@@ -153,6 +160,8 @@ FJavaAndroidMediaPlayer::FJavaAndroidMediaPlayer(bool swizzlePixels, bool vulkan
 	, DidResolutionChangeMethod(GetClassMethod("didResolutionChange", "()Z"))
 	, GetExternalTextureIdMethod(GetClassMethod("getExternalTextureId", "()I"))
 	, UpdateVideoFrameMethod(GetClassMethod("updateVideoFrame", "(I)Lcom/epicgames/unreal/MediaPlayer14$FrameUpdateInfo;"))
+	, SetUsingBitmapRendererImageReaderMethod(GetClassMethod("setUsingBitmapRendererImageReader", "(Z)V"))
+	, GetVideoLastMediaFrameDataMethod(GetClassMethod("getVideoLastMediaFrameData", "()Lcom/epicgames/unreal/BitmapRenderer$FrameData;"))
 {
 	VideoTexture = nullptr;
 	bVideoTextureValid = false;
@@ -168,6 +177,9 @@ FJavaAndroidMediaPlayer::FJavaAndroidMediaPlayer(bool swizzlePixels, bool vulkan
 		SelectTrackMethod = GetClassMethod("selectTrack", "(I)V");
 	}
 
+	BRSelection = BitmapRendererSelection::Unknown;
+	UsingVulkanRenderer = vulkanRenderer && FAndroidMisc::ShouldUseVulkan();
+
 	JNIEnv* JEnv = FAndroidApplication::GetJavaEnv();
 
 	// get field IDs for FrameUpdateInfo class members
@@ -370,8 +382,66 @@ void FJavaAndroidMediaPlayer::SetAudioVolume(float Volume)
 	CallMethod<void>(SetAudioVolumeMethod, Volume);
 }
 
+/** This function has to be called before calling any GetVideoLastXXXXXXX() functions */
+void FJavaAndroidMediaPlayer::SelectBitmapRenderer()
+{
+	if (BRSelection == BitmapRendererSelection::Unknown)
+	{
+		if (UsingVulkanRenderer && AndroidEnableMediaPlayerBRIR && FAndroidMisc::IsExternalMemoryAndroidHardwareBufferExtensionLoaded())
+		{
+			BRSelection = BitmapRendererSelection::ImageReader;
+		}
+		else
+		{
+			BRSelection = BitmapRendererSelection::Legacy;
+		}
+
+		JNIEnv* JEnv = FAndroidApplication::GetJavaEnv();
+		JEnv->CallVoidMethod(Object, SetUsingBitmapRendererImageReaderMethod.Method, (BRSelection == BitmapRendererSelection::ImageReader));
+	}
+}
+
+bool FJavaAndroidMediaPlayer::IsUsingBitmapRendererImageReader()
+{
+	SelectBitmapRenderer();
+	return (BRSelection == BitmapRendererSelection::ImageReader);
+}
+
+jobject FJavaAndroidMediaPlayer::GetVideoLastMediaFrameData()
+{
+	if (!IsUsingBitmapRendererImageReader())
+	{
+		return nullptr;
+	}
+
+	JNIEnv* JEnv = FAndroidApplication::GetJavaEnv();
+	jobject localFrameData = JEnv->CallObjectMethod(Object, GetVideoLastMediaFrameDataMethod.Method);
+	if (JEnv->ExceptionCheck())
+	{
+		JEnv->ExceptionDescribe();
+		JEnv->ExceptionClear();
+		return nullptr;
+	}
+
+	if (localFrameData == nullptr)
+	{
+		return nullptr;
+	}
+
+	jobject globalFrameData = FAndroidApplication::GetJavaEnv()->NewGlobalRef(localFrameData);
+
+	JEnv->DeleteLocalRef(localFrameData);
+
+	return globalFrameData;
+}
+
 bool FJavaAndroidMediaPlayer::GetVideoLastFrameData(void* & outPixels, int64 & outCount)
 {
+	if (IsUsingBitmapRendererImageReader())
+	{
+		return false;
+	}
+
 	// This can return an exception in some cases
 	JNIEnv*	JEnv = FAndroidApplication::GetJavaEnv();
 	auto buffer = NewScopedJavaObject(JEnv, JEnv->CallObjectMethod(Object, GetVideoLastFrameDataMethod.Method));
@@ -449,6 +519,11 @@ bool FJavaAndroidMediaPlayer::UpdateVideoFrame(int32 ExternalTextureId, int32 *C
 
 bool FJavaAndroidMediaPlayer::GetVideoLastFrame(int32 destTexture)
 {
+	if (IsUsingBitmapRendererImageReader())
+	{
+		return false;
+	}
+
 	// This can return an exception in some cases
 	JNIEnv*	JEnv = FAndroidApplication::GetJavaEnv();
 	bool Result = JEnv->CallBooleanMethod(Object, GetVideoLastFrameMethod.Method, destTexture);
diff --git a/Engine/Source/Runtime/ApplicationCore/Public/Android/AndroidJavaMediaFrameData.h b/Engine/Source/Runtime/ApplicationCore/Public/Android/AndroidJavaMediaFrameData.h
new file mode 100644
--- /dev/null
+++ b/Engine/Source/Runtime/ApplicationCore/Public/Android/AndroidJavaMediaFrameData.h
@@ -0,0 +1,58 @@
+// Copyright Epic Games, Inc. All Rights Reserved.
+
+#pragma once
+
+#if USE_ANDROID_JNI
+#include "CoreMinimal.h"
+#include "Android/AndroidJava.h"
+#include "IMediaTextureSample.h"
+#include "RHI.h"
+#include "RHIUtilities.h"
+
+class FAndroidJavaMediaFrameData
+{
+public:
+	FAndroidJavaMediaFrameData();
+	virtual ~FAndroidJavaMediaFrameData();
+
+	FAndroidJavaMediaFrameData(const FAndroidJavaMediaFrameData& Other) = delete;
+	FAndroidJavaMediaFrameData& operator=(const FAndroidJavaMediaFrameData& Other) = delete;
+
+	FAndroidJavaMediaFrameData& operator=(FAndroidJavaMediaFrameData& Other);
+
+	operator bool() const
+	{
+		return (FrameDataGlobalRef != nullptr);
+	}
+
+	bool IsReadyToClean()
+	{
+		if (Fence.IsValid())
+		{
+			return (Fence->Poll());
+		}
+		return true;
+	}
+
+	bool Set(jobject InFrameData);
+	void CleanUp();
+
+	jobject Extract(FVector2f& OutScale, FVector2f& OutOffset);
+	bool ExtractToTextureVulkan(FRHICommandListImmediate& RHICmdList, FTextureRHIRef& InDstTexture, IMediaTextureSample* TextureSample);
+	bool ExtractToTextureOES(FRHICommandListImmediate& RHICmdList, FTextureRHIRef& InDstTexture, IMediaTextureSample* TextureSample);
+
+private:
+	TRefCountPtr<FRHIGPUFence> Fence;
+
+	/** The java frame data from the surface provider*/
+	jobject FrameDataGlobalRef;
+
+	/** FrameData members / methods */
+	static jfieldID FrameData_HardwareBufferHandle;
+	static jfieldID FrameData_UScale;
+	static jfieldID FrameData_UOffset;
+	static jfieldID FrameData_VScale;
+	static jfieldID FrameData_VOffset;
+	static jmethodID FrameData_ReleaseFN;
+};
+#endif //USE_ANDROID_JNI
\ No newline at end of file
diff --git a/Engine/Source/Runtime/ApplicationCore/Public/Android/AndroidJavaMediaPlayer.h b/Engine/Source/Runtime/ApplicationCore/Public/Android/AndroidJavaMediaPlayer.h
--- a/Engine/Source/Runtime/ApplicationCore/Public/Android/AndroidJavaMediaPlayer.h
+++ b/Engine/Source/Runtime/ApplicationCore/Public/Android/AndroidJavaMediaPlayer.h
@@ -91,6 +91,8 @@ public:
 	void SetVideoEnabled(bool enabled = true);
 	void SetAudioEnabled(bool enabled = true);
 	void SetAudioVolume(float Volume);
+	bool IsUsingBitmapRendererImageReader();
+	jobject GetVideoLastMediaFrameData();
 	bool GetVideoLastFrameData(void* & outPixels, int64 & outCount);
 	void Start();
 	void Pause();
@@ -141,6 +143,9 @@ private:
 	FJavaClassMethod DidResolutionChangeMethod;
 	FJavaClassMethod GetExternalTextureIdMethod;
 	FJavaClassMethod UpdateVideoFrameMethod;
+	// new way of getting video frame data
+	FJavaClassMethod SetUsingBitmapRendererImageReaderMethod;
+	FJavaClassMethod GetVideoLastMediaFrameDataMethod;
 
 	// FrameUpdateInfo member field ids
 	jclass FrameUpdateInfoClass;
@@ -188,6 +193,18 @@ private:
 	/** Current media data source if using one. */
 	TSharedPtr<FJavaAndroidMediaDataSource, ESPMode::ThreadSafe> MediaDataSource;
 
+	enum BitmapRendererSelection 
+	{
+		Unknown,
+		Legacy,
+		ImageReader 
+	};
+
+	BitmapRendererSelection BRSelection; // BR - BitmapRenderer
+	bool UsingVulkanRenderer;
+
+	void SelectBitmapRenderer();
+
 public:
 	FTextureRHIRef GetVideoTexture()
 	{
diff --git a/Engine/Source/Runtime/Core/Core.Build.cs b/Engine/Source/Runtime/Core/Core.Build.cs
--- a/Engine/Source/Runtime/Core/Core.Build.cs
+++ b/Engine/Source/Runtime/Core/Core.Build.cs
@@ -18,6 +18,8 @@ public class Core : ModuleRules
 		if (Target.Platform == UnrealTargetPlatform.Android)
 		{
 			PrivateDependencyModuleNames.Add("GoogleGameSDK");
+			
+			PrivateIncludePathModuleNames.Add("VulkanRHI"); 				// Required for AndroidPlatformMisc
 
 			if (Target.Configuration != UnrealTargetConfiguration.Shipping)
 			{
diff --git a/Engine/Source/Runtime/Core/Private/Android/AndroidPlatformMisc.cpp b/Engine/Source/Runtime/Core/Private/Android/AndroidPlatformMisc.cpp
--- a/Engine/Source/Runtime/Core/Private/Android/AndroidPlatformMisc.cpp
+++ b/Engine/Source/Runtime/Core/Private/Android/AndroidPlatformMisc.cpp
@@ -56,6 +56,10 @@
 #include "Logging/LogMacros.h"
 #include "Misc/OutputDeviceError.h"
 
+#include "VulkanCommon.h"
+#include "IVulkanDynamicRHI.h"
+#include <vulkan/vulkan_android.h>
+
 #if USE_ANDROID_JNI
 extern AAssetManager * AndroidThunkCpp_GetAssetManager();
 extern int32 GAndroidPackageVersion;
@@ -2013,251 +2017,6 @@ namespace AndroidPlatformMisc
 #define VKAPI_PTR
 #endif
 
-#define VK_MAKE_VERSION(major, minor, patch) \
-	(((major) << 22) | ((minor) << 12) | (patch))
-
-#define VK_VERSION_MAJOR(version) ((uint32)(version) >> 22)
-#define VK_VERSION_MINOR(version) (((uint32)(version) >> 12) & 0x3ff)
-#define VK_VERSION_PATCH(version) ((uint32)(version) & 0xfff)
-
-typedef uint32 VkFlags;
-typedef uint32 VkBool32;
-
-#define VK_DEFINE_HANDLE(object) typedef struct object##_T* object;
-
-VK_DEFINE_HANDLE(VkInstance)
-VK_DEFINE_HANDLE(VkPhysicalDevice)
-
-typedef enum VkResult {
-	VK_SUCCESS = 0,
-	VK_NOT_READY = 1,
-	VK_TIMEOUT = 2,
-	VK_EVENT_SET = 3,
-	VK_EVENT_RESET = 4,
-	VK_INCOMPLETE = 5,
-	VK_ERROR_OUT_OF_HOST_MEMORY = -1,
-	VK_ERROR_OUT_OF_DEVICE_MEMORY = -2,
-	VK_ERROR_INITIALIZATION_FAILED = -3,
-	VK_ERROR_DEVICE_LOST = -4,
-	VK_ERROR_MEMORY_MAP_FAILED = -5,
-	VK_ERROR_LAYER_NOT_PRESENT = -6,
-	VK_ERROR_EXTENSION_NOT_PRESENT = -7,
-	VK_ERROR_FEATURE_NOT_PRESENT = -8,
-	VK_ERROR_INCOMPATIBLE_DRIVER = -9,
-	VK_ERROR_TOO_MANY_OBJECTS = -10,
-	VK_ERROR_FORMAT_NOT_SUPPORTED = -11,
-	VK_ERROR_SURFACE_LOST_KHR = -1000000000,
-	VK_ERROR_NATIVE_WINDOW_IN_USE_KHR = -1000000001,
-	VK_SUBOPTIMAL_KHR = 1000001003,
-	VK_ERROR_OUT_OF_DATE_KHR = -1000001004,
-	VK_ERROR_INCOMPATIBLE_DISPLAY_KHR = -1000003001,
-	VK_ERROR_VALIDATION_FAILED_EXT = -1000011001,
-	VK_ERROR_INVALID_SHADER_NV = -1000012000,
-	VK_RESULT_BEGIN_RANGE = VK_ERROR_FORMAT_NOT_SUPPORTED,
-	VK_RESULT_END_RANGE = VK_INCOMPLETE,
-	VK_RESULT_RANGE_SIZE = (VK_INCOMPLETE - VK_ERROR_FORMAT_NOT_SUPPORTED + 1),
-	VK_RESULT_MAX_ENUM = 0x7FFFFFFF
-} VkResult;
-
-typedef enum VkStructureType {
-	VK_STRUCTURE_TYPE_APPLICATION_INFO = 0,
-	VK_STRUCTURE_TYPE_INSTANCE_CREATE_INFO = 1,
-	VK_STRUCTURE_TYPE_MAX_ENUM = 0x7FFFFFFF
-} VkStructureType;
-
-typedef VkFlags VkInstanceCreateFlags;
-
-typedef struct VkApplicationInfo {
-	VkStructureType	sType;
-	const void*		pNext;
-	const char*		pApplicationName;
-	uint32			applicationVersion;
-	const char*		pEngineName;
-	uint			engineVersion;
-	uint			apiVersion;
-} VkApplicationInfo;
-
-typedef struct VkInstanceCreateInfo {
-	VkStructureType				sType;
-	const void*					pNext;
-	VkInstanceCreateFlags		flags;
-	const VkApplicationInfo*	pApplicationInfo;
-	uint32						enabledLayerCount;
-	const char* const*			ppEnabledLayerNames;
-	uint32						enabledExtensionCount;
-	const char* const*			ppEnabledExtensionNames;
-} VkInstanceCreateInfo;
-
-typedef struct VkAllocationCallbacks {
-	void*	pUserData;
-	void*	pfnAllocation;
-	void*	pfnReallocation;
-	void*	pfnFree;
-	void*	pfnInternalAllocation;
-	void*	pfnInternalFree;
-} VkAllocationCallbacks;
-
-typedef uint64 VkDeviceSize;
-typedef VkFlags VkSampleCountFlags;
-
-typedef struct VkPhysicalDeviceLimits {
-	uint32				maxImageDimension1D;
-	uint32				maxImageDimension2D;
-	uint32				maxImageDimension3D;
-	uint32				maxImageDimensionCube;
-	uint32				maxImageArrayLayers;
-	uint32				maxTexelBufferElements;
-	uint32				maxUniformBufferRange;
-	uint32				maxStorageBufferRange;
-	uint32				maxPushConstantsSize;
-	uint32				maxMemoryAllocationCount;
-	uint32				maxSamplerAllocationCount;
-	VkDeviceSize		bufferImageGranularity;
-	VkDeviceSize		sparseAddressSpaceSize;
-	uint32				maxBoundDescriptorSets;
-	uint32				maxPerStageDescriptorSamplers;
-	uint32				maxPerStageDescriptorUniformBuffers;
-	uint32				maxPerStageDescriptorStorageBuffers;
-	uint32				maxPerStageDescriptorSampledImages;
-	uint32				maxPerStageDescriptorStorageImages;
-	uint32				maxPerStageDescriptorInputAttachments;
-	uint32				maxPerStageResources;
-	uint32				maxDescriptorSetSamplers;
-	uint32				maxDescriptorSetUniformBuffers;
-	uint32				maxDescriptorSetUniformBuffersDynamic;
-	uint32				maxDescriptorSetStorageBuffers;
-	uint32				maxDescriptorSetStorageBuffersDynamic;
-	uint32				maxDescriptorSetSampledImages;
-	uint32				maxDescriptorSetStorageImages;
-	uint32				maxDescriptorSetInputAttachments;
-	uint32				maxVertexInputAttributes;
-	uint32				maxVertexInputBindings;
-	uint32				maxVertexInputAttributeOffset;
-	uint32				maxVertexInputBindingStride;
-	uint32				maxVertexOutputComponents;
-	uint32				maxTessellationGenerationLevel;
-	uint32				maxTessellationPatchSize;
-	uint32				maxTessellationControlPerVertexInputComponents;
-	uint32				maxTessellationControlPerVertexOutputComponents;
-	uint32				maxTessellationControlPerPatchOutputComponents;
-	uint32				maxTessellationControlTotalOutputComponents;
-	uint32				maxTessellationEvaluationInputComponents;
-	uint32				maxTessellationEvaluationOutputComponents;
-	uint32				maxGeometryShaderInvocations;
-	uint32				maxGeometryInputComponents;
-	uint32				maxGeometryOutputComponents;
-	uint32				maxGeometryOutputVertices;
-	uint32				maxGeometryTotalOutputComponents;
-	uint32				maxFragmentInputComponents;
-	uint32				maxFragmentOutputAttachments;
-	uint32				maxFragmentDualSrcAttachments;
-	uint32				maxFragmentCombinedOutputResources;
-	uint32				maxComputeSharedMemorySize;
-	uint32				maxComputeWorkGroupCount[3];
-	uint32				maxComputeWorkGroupInvocations;
-	uint32				maxComputeWorkGroupSize[3];
-	uint32				subPixelPrecisionBits;
-	uint32				subTexelPrecisionBits;
-	uint32				mipmapPrecisionBits;
-	uint32				maxDrawIndexedIndexValue;
-	uint32				maxDrawIndirectCount;
-	float				maxSamplerLodBias;
-	float				maxSamplerAnisotropy;
-	uint32				maxViewports;
-	uint32				maxViewportDimensions[2];
-	float				viewportBoundsRange[2];
-	uint32				viewportSubPixelBits;
-	size_t				minMemoryMapAlignment;
-	VkDeviceSize		minTexelBufferOffsetAlignment;
-	VkDeviceSize		minUniformBufferOffsetAlignment;
-	VkDeviceSize		minStorageBufferOffsetAlignment;
-	int32				minTexelOffset;
-	uint32				maxTexelOffset;
-	int32				minTexelGatherOffset;
-	uint32				maxTexelGatherOffset;
-	float				minInterpolationOffset;
-	float				maxInterpolationOffset;
-	uint32				subPixelInterpolationOffsetBits;
-	uint32				maxFramebufferWidth;
-	uint32				maxFramebufferHeight;
-	uint32				maxFramebufferLayers;
-	VkSampleCountFlags	framebufferColorSampleCounts;
-	VkSampleCountFlags	framebufferDepthSampleCounts;
-	VkSampleCountFlags	framebufferStencilSampleCounts;
-	VkSampleCountFlags	framebufferNoAttachmentsSampleCounts;
-	uint32				maxColorAttachments;
-	VkSampleCountFlags	sampledImageColorSampleCounts;
-	VkSampleCountFlags	sampledImageIntegerSampleCounts;
-	VkSampleCountFlags	sampledImageDepthSampleCounts;
-	VkSampleCountFlags	sampledImageStencilSampleCounts;
-	VkSampleCountFlags	storageImageSampleCounts;
-	uint32				maxSampleMaskWords;
-	VkBool32			timestampComputeAndGraphics;
-	float				timestampPeriod;
-	uint32				maxClipDistances;
-	uint32				maxCullDistances;
-	uint32				maxCombinedClipAndCullDistances;
-	uint32				discreteQueuePriorities;
-	float				pointSizeRange[2];
-	float				lineWidthRange[2];
-	float				pointSizeGranularity;
-	float				lineWidthGranularity;
-	VkBool32			strictLines;
-	VkBool32			standardSampleLocations;
-	VkDeviceSize		optimalBufferCopyOffsetAlignment;
-	VkDeviceSize		optimalBufferCopyRowPitchAlignment;
-	VkDeviceSize		nonCoherentAtomSize;
-} VkPhysicalDeviceLimits;
-
-typedef struct VkPhysicalDeviceSparseProperties {
-	VkBool32	residencyStandard2DBlockShape;
-	VkBool32	residencyStandard2DMultisampleBlockShape;
-	VkBool32	residencyStandard3DBlockShape;
-	VkBool32	residencyAlignedMipSize;
-	VkBool32	residencyNonResidentStrict;
-} VkPhysicalDeviceSparseProperties;
-
-typedef enum VkPhysicalDeviceType {
-	VK_PHYSICAL_DEVICE_TYPE_OTHER = 0,
-	VK_PHYSICAL_DEVICE_TYPE_INTEGRATED_GPU = 1,
-	VK_PHYSICAL_DEVICE_TYPE_DISCRETE_GPU = 2,
-	VK_PHYSICAL_DEVICE_TYPE_VIRTUAL_GPU = 3,
-	VK_PHYSICAL_DEVICE_TYPE_CPU = 4,
-	VK_PHYSICAL_DEVICE_TYPE_BEGIN_RANGE = VK_PHYSICAL_DEVICE_TYPE_OTHER,
-	VK_PHYSICAL_DEVICE_TYPE_END_RANGE = VK_PHYSICAL_DEVICE_TYPE_CPU,
-	VK_PHYSICAL_DEVICE_TYPE_RANGE_SIZE = (VK_PHYSICAL_DEVICE_TYPE_CPU - VK_PHYSICAL_DEVICE_TYPE_OTHER + 1),
-	VK_PHYSICAL_DEVICE_TYPE_MAX_ENUM = 0x7FFFFFFF
-} VkPhysicalDeviceType;
-
-#define VK_MAX_PHYSICAL_DEVICE_NAME_SIZE 256
-#define VK_UUID_SIZE 16
-
-typedef struct VkPhysicalDeviceProperties {
-	uint32								apiVersion;
-	uint32								driverVersion;
-	uint32								vendorID;
-	uint32								deviceID;
-	VkPhysicalDeviceType				deviceType;
-	char								deviceName[VK_MAX_PHYSICAL_DEVICE_NAME_SIZE];
-	uint8_t								pipelineCacheUUID[VK_UUID_SIZE];
-	VkPhysicalDeviceLimits				limits;
-	VkPhysicalDeviceSparseProperties	sparseProperties;
-} VkPhysicalDeviceProperties;
-
-#define VK_MAX_EXTENSION_NAME_SIZE        256
-#define VK_GOOGLE_DISPLAY_TIMING_EXTENSION_NAME "VK_GOOGLE_display_timing"
-
-typedef struct VkExtensionProperties {
-	char        extensionName[VK_MAX_EXTENSION_NAME_SIZE];
-	uint32_t    specVersion;
-} VkExtensionProperties;
-
-typedef VkResult(VKAPI_PTR *PFN_vkCreateInstance)(const VkInstanceCreateInfo* pCreateInfo, const VkAllocationCallbacks* pAllocator, VkInstance* pInstance);
-typedef void (VKAPI_PTR *PFN_vkDestroyInstance)(VkInstance instance, const VkAllocationCallbacks* pAllocator);
-typedef VkResult(VKAPI_PTR *PFN_vkEnumeratePhysicalDevices)(VkInstance instance, uint32* pPhysicalDeviceCount, VkPhysicalDevice* pPhysicalDevices);
-typedef void (VKAPI_PTR *PFN_vkGetPhysicalDeviceProperties)(VkPhysicalDevice physicalDevice, VkPhysicalDeviceProperties* pProperties);
-typedef VkResult(VKAPI_PTR *PFN_vkEnumerateDeviceExtensionProperties)(VkPhysicalDevice physicalDevice, const char* pLayerName, uint32_t* pPropertyCount, VkExtensionProperties* pProperties);
-
 ///////////////////////////////////////////////////////////////////////////////
 
 #define UE_VK_API_VERSION	VK_MAKE_VERSION(1, 1, 0)
@@ -2598,6 +2357,32 @@ FString FAndroidMisc::GetVulkanVersion()
 	return VulkanVersionString;
 }
 
+bool FAndroidMisc::IsExternalMemoryAndroidHardwareBufferExtensionLoaded()
+{
+	if (GIsRHIInitialized && ShouldUseVulkan())
+	{
+		IVulkanDynamicRHI* RHI = GetIVulkanDynamicRHI();
+		TArray<FAnsiString> LoadedDeviceExtensions = RHI->RHIGetLoadedDeviceExtensions();
+
+		// Consider replace the following check by using PhysicalDeviceFeatures.Core_1_1.samplerYcbcrConversion
+		if (!LoadedDeviceExtensions.Find(FAnsiString(VK_KHR_SAMPLER_YCBCR_CONVERSION_EXTENSION_NAME)))
+		{
+			UE_LOG(LogAndroid, Log, TEXT("Selecting CPU path because GPU extension '" VK_KHR_SAMPLER_YCBCR_CONVERSION_EXTENSION_NAME "' is not available!"));
+		}
+		else if (!LoadedDeviceExtensions.Find(FAnsiString(VK_ANDROID_EXTERNAL_MEMORY_ANDROID_HARDWARE_BUFFER_EXTENSION_NAME)))
+		{
+			UE_LOG(LogAndroid, Log, TEXT("Selecting CPU path because GPU extension '" VK_ANDROID_EXTERNAL_MEMORY_ANDROID_HARDWARE_BUFFER_EXTENSION_NAME "' is not available!"));
+		}
+		else
+		{
+			UE_LOG(LogAndroid, Log, TEXT("Selecting GPU path because it is enabled via Electra.AndroidUseGpuOutputPath = 1"));
+			return true;
+		}
+	}
+
+	return false;
+}
+
 const TMap<FString, FString>& FAndroidMisc::GetConfigRuleVars()
 {
 	return ConfigRulesVariables;
diff --git a/Engine/Source/Runtime/Core/Public/Android/AndroidPlatformMisc.h b/Engine/Source/Runtime/Core/Public/Android/AndroidPlatformMisc.h
--- a/Engine/Source/Runtime/Core/Public/Android/AndroidPlatformMisc.h
+++ b/Engine/Source/Runtime/Core/Public/Android/AndroidPlatformMisc.h
@@ -219,6 +219,7 @@ public:
 	static CORE_API bool ShouldUseVulkan();
 	static CORE_API bool ShouldUseDesktopVulkan();
 	static CORE_API FString GetVulkanVersion();
+	static CORE_API bool IsExternalMemoryAndroidHardwareBufferExtensionLoaded();
 	typedef TFunction<void(void* NewNativeHandle)> ReInitWindowCallbackType;
 	static CORE_API ReInitWindowCallbackType GetOnReInitWindowCallback();
 	static CORE_API void SetOnReInitWindowCallback(ReInitWindowCallbackType InOnReInitWindowCallback);
diff --git a/Engine/Source/Runtime/OpenGLDrv/Private/Android/AndroidOpenGL.cpp b/Engine/Source/Runtime/OpenGLDrv/Private/Android/AndroidOpenGL.cpp
--- a/Engine/Source/Runtime/OpenGLDrv/Private/Android/AndroidOpenGL.cpp
+++ b/Engine/Source/Runtime/OpenGLDrv/Private/Android/AndroidOpenGL.cpp
@@ -404,6 +404,11 @@ void FPlatformOpenGLDevice::LoadEXT()
 	eglGetCompositorTimingSupportedANDROID_p = (PFNeglQueryTimestampSupportedANDROID)((void*)eglGetProcAddress("eglGetCompositorTimingSupportedANDROID"));
 	eglGetFrameTimestampsSupportedANDROID_p = (PFNeglQueryTimestampSupportedANDROID)((void*)eglGetProcAddress("eglGetFrameTimestampsSupportedANDROID"));
 
+	eglGetNativeClientBufferANDROID_p = (PFNEGLGETNATIVECLIENTBUFFERANDROIDPROC)((void*)eglGetProcAddress("eglGetNativeClientBufferANDROID"));
+	eglCreateImageKHR_p = (PFNEGLCREATEIMAGEKHRPROC)((void*)eglGetProcAddress("eglCreateImageKHR"));
+	eglDestroyImageKHR_p = (PFNEGLDESTROYIMAGEKHRPROC)((void*)eglGetProcAddress("eglDestroyImageKHR"));
+	glEGLImageTargetTexture2DOES_p = (PFNGLEGLIMAGETARGETTEXTURE2DOESPROC)((void*)eglGetProcAddress("glEGLImageTargetTexture2DOES"));
+
 	const TCHAR* NotAvailable = TEXT("NOT Available");
 	const TCHAR* Present = TEXT("Present");
 
@@ -415,6 +420,11 @@ void FPlatformOpenGLDevice::LoadEXT()
 	UE_LOG(LogRHI, Log, TEXT("Extension %s %s"), TEXT("eglGetCompositorTimingSupportedANDROID"), eglGetCompositorTimingSupportedANDROID_p ? Present : NotAvailable);
 	UE_LOG(LogRHI, Log, TEXT("Extension %s %s"), TEXT("eglGetFrameTimestampsSupportedANDROID"), eglGetFrameTimestampsSupportedANDROID_p ? Present : NotAvailable);
 
+	UE_LOG(LogRHI, Log, TEXT("Extension %s %s"), TEXT("eglGetNativeClientBufferANDROID"), eglGetNativeClientBufferANDROID_p ? Present : NotAvailable);
+	UE_LOG(LogRHI, Log, TEXT("Extension %s %s"), TEXT("eglCreateImageKHR"), eglCreateImageKHR_p ? Present : NotAvailable);
+	UE_LOG(LogRHI, Log, TEXT("Extension %s %s"), TEXT("eglCreateImageKHR"), eglCreateImageKHR_p ? Present : NotAvailable);
+	UE_LOG(LogRHI, Log, TEXT("Extension %s %s"), TEXT("glEGLImageTargetTexture2DOES"), glEGLImageTargetTexture2DOES_p ? Present : NotAvailable);
+
 	glDebugMessageControlKHR = (PFNGLDEBUGMESSAGECONTROLKHRPROC)((void*)eglGetProcAddress("glDebugMessageControlKHR"));
 
 	// Some PowerVR drivers (Rogue Han and Intel-based devices) are crashing using glDebugMessageControlKHR (causes signal 11 crash)
diff --git a/Engine/Source/Runtime/OpenGLDrv/Private/Android/AndroidOpenGLFunctions.cpp b/Engine/Source/Runtime/OpenGLDrv/Private/Android/AndroidOpenGLFunctions.cpp
--- a/Engine/Source/Runtime/OpenGLDrv/Private/Android/AndroidOpenGLFunctions.cpp
+++ b/Engine/Source/Runtime/OpenGLDrv/Private/Android/AndroidOpenGLFunctions.cpp
@@ -16,6 +16,11 @@ PFNEGLDESTROYSYNCKHRPROC eglDestroySyncKHR_p = NULL;
 PFNEGLCLIENTWAITSYNCKHRPROC eglClientWaitSyncKHR_p = NULL;
 PFNEGLGETSYNCATTRIBKHRPROC eglGetSyncAttribKHR_p = NULL;
 
+PFNEGLGETNATIVECLIENTBUFFERANDROIDPROC eglGetNativeClientBufferANDROID_p = nullptr;
+PFNEGLCREATEIMAGEKHRPROC eglCreateImageKHR_p = nullptr;
+PFNEGLDESTROYIMAGEKHRPROC eglDestroyImageKHR_p = nullptr;
+PFNGLEGLIMAGETARGETTEXTURE2DOESPROC glEGLImageTargetTexture2DOES_p = nullptr;
+
 namespace GLFuncPointers
 {
 	// Offscreen MSAA rendering
diff --git a/Engine/Source/Runtime/OpenGLDrv/Private/Android/AndroidOpenGLFunctions.h b/Engine/Source/Runtime/OpenGLDrv/Private/Android/AndroidOpenGLFunctions.h
--- a/Engine/Source/Runtime/OpenGLDrv/Private/Android/AndroidOpenGLFunctions.h
+++ b/Engine/Source/Runtime/OpenGLDrv/Private/Android/AndroidOpenGLFunctions.h
@@ -12,6 +12,8 @@ typedef GLboolean(GL_APIENTRYP PFNeglGetCompositorTimingANDROID) (EGLDisplay dpy
 typedef GLboolean(GL_APIENTRYP PFNeglGetFrameTimestampsANDROID) (EGLDisplay dpy, EGLSurface surface, EGLuint64KHR frameId, EGLint numTimestamps, const EGLint* timestamps, EGLnsecsANDROID* values);
 typedef GLboolean(GL_APIENTRYP PFNeglQueryTimestampSupportedANDROID) (EGLDisplay dpy, EGLSurface surface, EGLint timestamp);
 
+typedef EGLClientBuffer(GL_APIENTRYP PFNEGLGETNATIVECLIENTBUFFERANDROIDPROC) (const struct AHardwareBuffer* buffer);
+
 extern "C"
 {
 	extern PFNEGLGETSYSTEMTIMENVPROC eglGetSystemTimeNV_p;
@@ -27,6 +29,11 @@ extern "C"
 	extern PFNeglQueryTimestampSupportedANDROID eglQueryTimestampSupportedANDROID_p;
 	extern PFNeglQueryTimestampSupportedANDROID eglGetCompositorTimingSupportedANDROID_p;
 	extern PFNeglQueryTimestampSupportedANDROID eglGetFrameTimestampsSupportedANDROID_p;
+
+	extern PFNEGLGETNATIVECLIENTBUFFERANDROIDPROC eglGetNativeClientBufferANDROID_p;
+	extern PFNEGLCREATEIMAGEKHRPROC eglCreateImageKHR_p;
+	extern PFNEGLDESTROYIMAGEKHRPROC eglDestroyImageKHR_p;
+	extern PFNGLEGLIMAGETARGETTEXTURE2DOESPROC glEGLImageTargetTexture2DOES_p;
 }
 
 namespace GLFuncPointers
diff --git a/Engine/Source/Runtime/OpenGLDrv/Private/OpenGLDevice.cpp b/Engine/Source/Runtime/OpenGLDrv/Private/OpenGLDevice.cpp
--- a/Engine/Source/Runtime/OpenGLDrv/Private/OpenGLDevice.cpp
+++ b/Engine/Source/Runtime/OpenGLDrv/Private/OpenGLDevice.cpp
@@ -114,7 +114,7 @@ void FOpenGLDynamicRHI::RHIEndFrame(const FRHIEndFrameArgs& Args)
 
 #if PLATFORM_ANDROID
 
-JNI_METHOD void Java_com_epicgames_unreal_MediaPlayer14_nativeClearCachedAttributeState(JNIEnv* jenv, jobject thiz, jint PositionAttrib, jint TexCoordsAttrib)
+JNI_METHOD void Java_com_epicgames_unreal_BitmapRendererLegacy_nativeClearCachedAttributeState(JNIEnv* jenv, jobject thiz, jint PositionAttrib, jint TexCoordsAttrib)
 {
 	FOpenGLDynamicRHI::Get().ClearCachedAttributeState(PositionAttrib, TexCoordsAttrib);
 }
diff --git a/Engine/Source/Runtime/OpenGLDrv/Private/OpenGLDrv.h b/Engine/Source/Runtime/OpenGLDrv/Private/OpenGLDrv.h
--- a/Engine/Source/Runtime/OpenGLDrv/Private/OpenGLDrv.h
+++ b/Engine/Source/Runtime/OpenGLDrv/Private/OpenGLDrv.h
@@ -268,6 +268,10 @@ public:
 	virtual int32 RHIGetGLMinorVersion() const final override;
 	virtual bool RHISupportsFramebufferSRGBEnable() const final override;
 	virtual FTextureRHIRef RHICreateTexture2DFromResource(EPixelFormat Format, uint32 SizeX, uint32 SizeY, uint32 NumMips, uint32 NumSamples, uint32 NumSamplesTileMem, const FClearValueBinding& ClearValueBinding, GLuint Resource, ETextureCreateFlags Flags) final override;
+#if PLATFORM_ANDROID
+	virtual FTextureRHIRef RHICreateTexture2DFromAndroidHardwareBuffer(FRHICommandListBase& RHICmdList, AHardwareBuffer* HardwareBuffer) override;
+#endif //PLATFORM_ANDROID
+
 	virtual FTextureRHIRef RHICreateTexture2DArrayFromResource(EPixelFormat Format, uint32 SizeX, uint32 SizeY, uint32 ArraySize, uint32 NumMips, uint32 NumSamples, uint32 NumSamplesTileMem, const FClearValueBinding& ClearValueBinding, GLuint Resource, ETextureCreateFlags Flags) final override;
 	virtual FTextureRHIRef RHICreateTextureCubeFromResource(EPixelFormat Format, uint32 Size, bool bArray, uint32 ArraySize, uint32 NumMips, uint32 NumSamples, uint32 NumSamplesTileMem, const FClearValueBinding& ClearValueBinding, GLuint Resource, ETextureCreateFlags Flags) final override;
 	virtual GLuint RHIGetResource(FRHITexture* InTexture) const final override;
diff --git a/Engine/Source/Runtime/OpenGLDrv/Private/OpenGLResources.h b/Engine/Source/Runtime/OpenGLDrv/Private/OpenGLResources.h
--- a/Engine/Source/Runtime/OpenGLDrv/Private/OpenGLResources.h
+++ b/Engine/Source/Runtime/OpenGLDrv/Private/OpenGLResources.h
@@ -944,6 +944,11 @@ class OPENGLDRV_API FOpenGLTexture : public FRHITexture, public FOpenGLViewableR
 public:
 	// Standard constructor.
 	explicit FOpenGLTexture(FRHICommandListBase& RHICmdList, FOpenGLTextureCreateDesc const& CreateDesc);
+#if PLATFORM_ANDROID
+	// Constructor for external hardware buffer.
+	explicit FOpenGLTexture(FRHICommandListBase& RHICmdList, FOpenGLTextureCreateDesc const& CreateDesc, 
+		const AHardwareBuffer_Desc& HardwareBufferDesc, AHardwareBuffer* HardwareBuffer);
+#endif //PLATFORM_ANDROID
 
 	// Constructor for external resources (RHICreateTexture2DFromResource etc).
 	explicit FOpenGLTexture(FOpenGLTextureCreateDesc const& CreateDesc, GLuint Resource);
@@ -1105,6 +1110,11 @@ public:
 
 	TArray<TRefCountPtr<FOpenGLPixelBuffer>> PixelBuffers;
 
+#if PLATFORM_ANDROID
+	// The image created from an external hardware buffer.
+	EGLImageKHR HardwareBufferImage;
+#endif //PLATFORM_ANDROID
+
 private:
 	/** Bitfields marking whether we have allocated storage for each mip */
 	TBitArray<TInlineAllocator<1> > bAllocatedStorage;
diff --git a/Engine/Source/Runtime/OpenGLDrv/Private/OpenGLTexture.cpp b/Engine/Source/Runtime/OpenGLDrv/Private/OpenGLTexture.cpp
--- a/Engine/Source/Runtime/OpenGLDrv/Private/OpenGLTexture.cpp
+++ b/Engine/Source/Runtime/OpenGLDrv/Private/OpenGLTexture.cpp
@@ -330,6 +330,69 @@ FOpenGLTexture::FOpenGLTexture(FOpenGLTextureCreateDesc const& CreateDesc, GLuin
 	, bMultisampleRenderbuffer(CreateDesc.bMultisampleRenderbuffer)
 {}
 
+#if PLATFORM_ANDROID
+// Constructor for external hardware buffer.
+FOpenGLTexture::FOpenGLTexture(FRHICommandListBase& RHICmdList, FOpenGLTextureCreateDesc const& CreateDesc, 
+	const AHardwareBuffer_Desc& HardwareBufferDesc, AHardwareBuffer* HardwareBuffer)
+	: FRHITexture(CreateDesc)
+	, Target(CreateDesc.Target)
+	, Attachment(CreateDesc.Attachment)
+	, MemorySize(CreateDesc.MemorySize)
+	, bIsPowerOfTwo(CreateDesc.bIsPowerOfTwo)
+	, bCanCreateAsEvicted(CreateDesc.bCanCreateAsEvicted)
+	, bStreamable(CreateDesc.bStreamable)
+	, bCubemap(CreateDesc.bCubemap)
+	, bArrayTexture(CreateDesc.bArrayTexture)
+	, bDepthStencil(CreateDesc.bDepthStencil)
+	, bAlias(false)
+	, bMultisampleRenderbuffer(CreateDesc.bMultisampleRenderbuffer)
+	, HardwareBufferImage(nullptr)
+{
+	QUICK_SCOPE_CYCLE_COUNTER(STAT_OpenGLCreateTextureTime);
+
+	RHICmdList.EnqueueLambda([this, HardwareBuffer](FRHICommandListBase& CmdList)
+		{
+			check(IsInRHIThread());
+
+			VERIFY_GL_SCOPE();
+			QUICK_SCOPE_CYCLE_COUNTER(STAT_FRHICommandUpdateDecoderExternaTexture_Execute);
+
+			ensure(eglGetNativeClientBufferANDROID_p != nullptr);
+			ensure(eglCreateImageKHR_p != nullptr);
+			ensure(eglDestroyImageKHR_p != nullptr);
+			ensure(glEGLImageTargetTexture2DOES_p != nullptr);
+
+			// Ext. EGL_ANDROID_get_native_client_buffer
+			EGLClientBuffer NativeClientBuffer = eglGetNativeClientBufferANDROID_p(HardwareBuffer);
+			if (NativeClientBuffer == nullptr)
+			{
+				UE_LOG(LogRHI, Warning, TEXT("eglGetNativeClientBufferANDROID - Could not get native client buffer!"));
+				return;
+			}
+
+			// Ext. EGL_ANDROID_image_native_buffer
+			EGLImageKHR EglImage = eglCreateImageKHR_p(eglGetCurrentDisplay(), EGL_NO_CONTEXT, EGL_NATIVE_BUFFER_ANDROID, NativeClientBuffer, nullptr);
+			if (EglImage == nullptr)
+			{
+				UE_LOG(LogRHI, Warning, TEXT("eglCreateImageKHR - Could not create EGLimage from native client buffer! B=0x%x E=0x%x"), NativeClientBuffer, eglGetError());
+				return;
+			}
+
+			this->HardwareBufferImage = EglImage;
+
+			glEGLImageTargetTexture2DOES_p(GL_TEXTURE_EXTERNAL_OES, this->HardwareBufferImage);
+			check(eglGetError() == EGL_SUCCESS);
+			check(glGetError() == 0);
+
+			// Bind the RHI texture's GL texture to the external image we got data about now...
+			glBindTexture(GL_TEXTURE_EXTERNAL_OES, *reinterpret_cast<int32*>(this->GetNativeResource()));
+			check(glGetError() == 0);
+		});
+
+	UpdateTextureStats(this, true);
+}
+#endif//PLATFORM_ANDROID
+
 // Standard constructor.
 FOpenGLTexture::FOpenGLTexture(FRHICommandListBase& RHICmdList, FOpenGLTextureCreateDesc const& CreateDesc)
 	: FRHITexture        (CreateDesc)
@@ -416,6 +479,13 @@ void FOpenGLTexture::DeleteGLResource()
 	VERIFY_GL_SCOPE();
 	QUICK_SCOPE_CYCLE_COUNTER(STAT_OpenGLDeleteGLTextureTime);
 
+#if PLATFORM_ANDROID
+	if (HardwareBufferImage != nullptr)
+	{
+		eglDestroyImageKHR_p(eglGetCurrentDisplay(), HardwareBufferImage);
+	}
+#endif //PLATFORM_ANDROID
+
 	if (Resource != 0)
 	{
 		switch (Target)
@@ -1837,6 +1907,31 @@ FTextureRHIRef FOpenGLDynamicRHI::RHICreateTexture2DFromResource(EPixelFormat Fo
 	return new FOpenGLTexture(Desc, Resource);
 }
 
+#if PLATFORM_ANDROID
+FTextureRHIRef FOpenGLDynamicRHI::RHICreateTexture2DFromAndroidHardwareBuffer(FRHICommandListBase& RHICmdList, AHardwareBuffer* HardwareBuffer)
+{
+	check(HardwareBuffer);
+
+	AHardwareBuffer_Desc HardwareBufferDesc;
+	AHardwareBuffer_describe(HardwareBuffer, &HardwareBufferDesc);
+	check((HardwareBufferDesc.usage & AHARDWAREBUFFER_USAGE_GPU_SAMPLED_IMAGE) != 0);
+	
+	// OpenGL driver will auto convert the color space for YUV (tested 420 and P010).
+	// But we should test with different video formats to identify what is not supported.
+	// UE_LOG(LogRHI, Log, TEXT("[Media] hardware buffer format: 0x%x"), HardwareBufferDesc.format & 0xff);
+
+	// Temp for avoiding a checking failure of PF_Unknown in FRHITextureDesc::Validate()
+	EPixelFormat PixelFormat = PF_R8G8B8A8;
+
+	const FRHITextureCreateDesc Desc =
+		FRHITextureCreateDesc::Create2D(TEXT("OpenGLTexture2DFromAndroidHardwareBuffer"), HardwareBufferDesc.width, HardwareBufferDesc.height, PixelFormat)
+		.SetFlags(ETextureCreateFlags::External)
+		.DetermineInititialState();
+
+	return new FOpenGLTexture(RHICmdList, Desc, HardwareBufferDesc, HardwareBuffer);
+}
+#endif //PLATFORM_ANDROID
+
 FTextureRHIRef FOpenGLDynamicRHI::RHICreateTexture2DArrayFromResource(EPixelFormat Format, uint32 SizeX, uint32 SizeY, uint32 ArraySize, uint32 NumMips, uint32 NumSamples, uint32 NumSamplesTileMem, const FClearValueBinding& ClearValueBinding, GLuint Resource, ETextureCreateFlags TexCreateFlags)
 {
 	const FRHITextureCreateDesc Desc =
diff --git a/Engine/Source/Runtime/OpenGLDrv/Public/IOpenGLDynamicRHI.h b/Engine/Source/Runtime/OpenGLDrv/Public/IOpenGLDynamicRHI.h
--- a/Engine/Source/Runtime/OpenGLDrv/Public/IOpenGLDynamicRHI.h
+++ b/Engine/Source/Runtime/OpenGLDrv/Public/IOpenGLDynamicRHI.h
@@ -17,6 +17,9 @@ struct IOpenGLDynamicRHI : public FDynamicRHIPSOFallback
 	virtual FTextureRHIRef RHICreateTexture2DFromResource(EPixelFormat Format, uint32 SizeX, uint32 SizeY, uint32 NumMips, uint32 NumSamples, uint32 NumSamplesTileMem, const FClearValueBinding& ClearValueBinding, GLuint Resource, ETextureCreateFlags Flags) = 0;
 	virtual FTextureRHIRef RHICreateTexture2DArrayFromResource(EPixelFormat Format, uint32 SizeX, uint32 SizeY, uint32 ArraySize, uint32 NumMips, uint32 NumSamples, uint32 NumSamplesTileMem, const FClearValueBinding& ClearValueBinding, GLuint Resource, ETextureCreateFlags Flags) = 0;
 	virtual FTextureRHIRef RHICreateTextureCubeFromResource(EPixelFormat Format, uint32 Size, bool bArray, uint32 ArraySize, uint32 NumMips, uint32 NumSamples, uint32 NumSamplesTileMem, const FClearValueBinding& ClearValueBinding, GLuint Resource, ETextureCreateFlags Flags) = 0;
+#if PLATFORM_ANDROID
+	virtual FTextureRHIRef RHICreateTexture2DFromAndroidHardwareBuffer(FRHICommandListBase& RHICmdList, AHardwareBuffer* HardwareBuffer) = 0;
+#endif //PLATFORM_ANDROID
 
 	virtual GLuint RHIGetResource(FRHITexture* InTexture) const = 0;
 	virtual bool RHIIsValidTexture(GLuint InTexture) const = 0;
