GameActivite:  move project configuration preparation to onCreatePreBody method
make processSystemInfo method protected

#rnx

diff --git a/Engine/Build/Android/Java/src/com/epicgames/unreal/GameActivity.java.template b/Engine/Build/Android/Java/src/com/epicgames/unreal/GameActivity.java.template
--- a/Engine/Build/Android/Java/src/com/epicgames/unreal/GameActivity.java.template
+++ b/Engine/Build/Android/Java/src/com/epicgames/unreal/GameActivity.java.template
@@ -1818,7 +1818,7 @@
 		return bConditionTrue;
 	}
 
-	private void processSystemInfo()
+	protected void processSystemInfo()
 	{
 		SystemInfoLogs = new LinkedList<String>();
 
@@ -2781,8 +2781,53 @@
 		activityContext = getActivityOfContext(this);
 		assert activityContext != null;
 
+		onCreatePreBody();
+
+		// get file paths
+		AssetManagerReference = this.getAssets();
+		InternalFilesDir = getFilesDir().getAbsolutePath() + "/";
+		ExternalFilesDir = Objects.requireNonNull(getExternalFilesDir(null)).getAbsolutePath() + "/";
+
+		ProcessSystemInfoThread = new Thread(new Runnable() {
+            public void run() {
+				processSystemInfo();
+			}
+		});
+		ProcessSystemInfoThread.setName("ProcessSysInfo");
+		ProcessSystemInfoThread.start();
+
+//$${gameActivityOnCreateBeforeLoadLibrariesAdditions}$$
+		LoadLibraries();
+
+		// Bug with certain devices calling onCreate twice, even with singleTask set and System.exit called in onDestroy.
+		if (nativeIsShippingBuild())
+		{
+			if (_activity != null)
+			{
+				Log.debug("LIFECYCLE: GameActivity.onCreate called twice!");
+				System.exit(0);
+			}
+		}
+		else
+		{
+			assert _activity == null;
+		}
+		
+		_activity = this;
+		bOnCreateCalled = true;
+		
+//$${gameActivityOnCreateBeginningAdditions}$$
+		Logger.RegisterCallback(this);
+
+		super.onCreate(savedInstanceState);
+
+		onCreateBody(savedInstanceState);
+	}
+
+	protected void onCreatePreBody()
+	{
 		// Retrieve version code, name, and target SDK
-		try 
+		try
 		{
 			PackageInfo packageInfo = getPackageManager().getPackageInfo(getPackageName(), 0);
 			VersionCode = packageInfo.versionCode;
@@ -2802,7 +2847,7 @@
 			Bundle bundle = ai.metaData;
 			_bundle = ai.metaData;
 
-			if ((ai.flags & ApplicationInfo.FLAG_DEBUGGABLE) == 0) 
+			if ((ai.flags & ApplicationInfo.FLAG_DEBUGGABLE) == 0)
 			{
 				IsForDistribution = true;
 			}
@@ -2873,46 +2918,6 @@
 		{
 			Log.debug( "Failed to load meta-data: NullPointer: " + e.getMessage());
 		}
-
-		// get file paths
-		AssetManagerReference = this.getAssets();
-		InternalFilesDir = getFilesDir().getAbsolutePath() + "/";
-		ExternalFilesDir = Objects.requireNonNull(getExternalFilesDir(null)).getAbsolutePath() + "/";
-
-		ProcessSystemInfoThread = new Thread(new Runnable() {
-            public void run() {
-				processSystemInfo();
-			}
-		});
-		ProcessSystemInfoThread.setName("ProcessSysInfo");
-		ProcessSystemInfoThread.start();
-
-//$${gameActivityOnCreateBeforeLoadLibrariesAdditions}$$
-		LoadLibraries();
-
-		// Bug with certain devices calling onCreate twice, even with singleTask set and System.exit called in onDestroy.
-		if (nativeIsShippingBuild())
-		{
-			if (_activity != null)
-			{
-				Log.debug("LIFECYCLE: GameActivity.onCreate called twice!");
-				System.exit(0);
-			}
-		}
-		else
-		{
-			assert _activity == null;
-		}
-		
-		_activity = this;
-		bOnCreateCalled = true;
-		
-//$${gameActivityOnCreateBeginningAdditions}$$
-		Logger.RegisterCallback(this);
-
-		super.onCreate(savedInstanceState);
-
-		onCreateBody(savedInstanceState);
 	}
 
 	protected void onCreateBody(Bundle savedInstanceState)
