Fix for RenderDoc failing to hook OpenGLES
Support for Java Perfetto markers (bJavaPerfettoMarkers=True) with Trace. filtering
Optimizes to runConfigRules and better thread overlap

#android

diff --git a/Engine/Build/Android/Java/src/com/epicgames/unreal/GameActivity.java.template b/Engine/Build/Android/Java/src/com/epicgames/unreal/GameActivity.java.template
--- a/Engine/Build/Android/Java/src/com/epicgames/unreal/GameActivity.java.template
+++ b/Engine/Build/Android/Java/src/com/epicgames/unreal/GameActivity.java.template
@@ -43,6 +43,7 @@ import java.util.Map;
 import java.util.HashMap;
 import java.util.ArrayList;
 import java.util.LinkedList;
+import java.util.Stack;
 import java.text.DecimalFormat;
 import java.text.DecimalFormatSymbols;
 import java.text.BreakIterator;
@@ -50,6 +51,8 @@ import java.util.zip.Inflater;
 import java.util.Stack;
 import java.util.Timer;
 import java.util.TimerTask;
+import java.util.concurrent.locks.Lock;
+import java.util.concurrent.locks.ReentrantLock;
 import java.lang.Thread;
 import java.lang.Runnable;
 import java.util.zip.CRC32;
@@ -360,12 +363,18 @@ public class GameActivity extends $${gameActivitySuperClass}$$ implements Surfac
 	protected static Bundle _extrasBundle;
 	protected boolean bOnCreateCalled = false;
 
-	Thread ProcessSystemInfoThread = null;
+	protected int MainThreadId;
+	protected int MainThreadPriority;
+
+	protected Thread ProcessSystemInfoThread = null;
+	private Lock SystemInfoLogsLock = new ReentrantLock();
+	private Lock ProcessSystemInfoLock = new ReentrantLock();
+	protected boolean ConfigRulesRunResult = true;
 	protected long ConfigRulesCRC = 0;
 	protected long ConfigRulesVarsCRC = 0;
 	protected Map<String, String> ConfigRulesVars = null;
-	private LinkedList<String> SystemInfoLogs = null;
-	Map<String, String> CpuInfo = null;
+	protected LinkedList<String> SystemInfoLogs = null;
+	protected Map<String, String> CpuInfo = null;
 
 	private HandlerThread memoryHandlerThread;
 	private Handler memoryHandler;
@@ -570,6 +579,71 @@ public class GameActivity extends $${gameActivitySuperClass}$$ implements Surfac
 	}
 	private EAlertDialogType CurrentDialogType = EAlertDialogType.None;
 
+	private void SystemInfo_InitLog()
+	{
+		SystemInfoLogsLock.lock();
+		SystemInfoLogs = new LinkedList<String>();
+		SystemInfoLogsLock.unlock();
+	}
+
+	private void SystemInfo_Log(String message)
+	{
+		// buffer logs if capture active, otherwise just log it
+		SystemInfoLogsLock.lock();
+		if (SystemInfoLogs == null)
+		{
+			Log.debug(message);
+		}
+		else
+		{
+			SystemInfoLogs.add(message);
+		}
+		SystemInfoLogsLock.unlock();
+	}
+
+	private void SystemInfo_LogMap(String caption, Map<String, String> map)
+	{
+		// buffer logs if capture active, otherwise just log it
+		SystemInfoLogsLock.lock();
+		if (SystemInfoLogs == null)
+		{
+			if (caption != null)
+			{
+				Log.debug(caption);
+			}
+			for (Map.Entry<String, String> entry : map.entrySet())
+			{
+				Log.debug("  " + entry.getKey() + " = " + entry.getValue());
+			}
+		}
+		else
+		{
+			if (caption != null)
+			{
+				SystemInfoLogs.add(caption);
+			}
+			for (Map.Entry<String, String> entry : map.entrySet())
+			{
+				SystemInfoLogs.add("  " + entry.getKey() + " = " + entry.getValue());
+			}
+		}
+		SystemInfoLogsLock.unlock();
+	}
+
+	private void SystemInfo_DumpLog()
+	{
+		SystemInfoLogsLock.lock();
+		if (SystemInfoLogs != null)
+		{
+			for (String logline : SystemInfoLogs)
+			{
+				Log.debug(logline);
+			}
+			SystemInfoLogs = null;
+		}
+		SystemInfoLogsLock.unlock();
+	}
+
 	
 	public boolean IsInVRMode()
 	{
@@ -779,6 +853,7 @@ public class GameActivity extends $${gameActivitySuperClass}$$ implements Surfac
                 result.add(packageName);
             }
             bufferedReader.close();
+			bufferedReader = null;
             process.waitFor();
         }
 		catch (Exception e)
@@ -799,8 +874,62 @@ public class GameActivity extends $${gameActivitySuperClass}$$ implements Surfac
 
     public boolean AndroidThunkJava_IsPackageInstalled(String packagename)
 	{
-        List<String> packages = getInstalledPackages(getApplicationContext());
-        return packages.contains(packagename);
+		final Context context = getApplicationContext();
+        final PackageManager pm = context.getPackageManager();
+        try
+		{
+            List<PackageInfo> apps = pm.getInstalledPackages(0);
+            for (PackageInfo packageInfo : apps) {
+				if (packageInfo.packageName.equals(packagename))
+				{
+					return true;
+				}
+            }
+			return false;
+        }
+		catch (Exception ignored)
+		{
+            // could have failed depending on bugs in Android version
+        }
+
+        // use fallback
+        BufferedReader bufferedReader = null;
+        try
+		{
+            Process process = Runtime.getRuntime().exec("pm list packages");
+            bufferedReader = new BufferedReader(new InputStreamReader(process.getInputStream()));
+            String line;
+            while ((line = bufferedReader.readLine()) != null) {
+                final String linePackageName = line.substring(line.indexOf(':') + 1);
+				if (linePackageName.equals(packagename))
+				{
+		            try
+					{
+						bufferedReader.close();
+					}
+					catch (Exception e) {
+					}
+					return true;
+				}
+            }
+            bufferedReader.close();
+			bufferedReader = null;
+            process.waitFor();
+        }
+		catch (Exception e)
+		{
+            e.printStackTrace();
+        }
+        if (bufferedReader != null)
+        {
+            try
+			{
+                bufferedReader.close();
+            }
+            catch (Exception e) {
+            }
+        }
+		return false;
     }
 
 	public String AndroidThunkJava_GetCommandLine()
@@ -860,7 +989,7 @@ public class GameActivity extends $${gameActivitySuperClass}$$ implements Surfac
 		if (bUseExternalFilesDir)
 		{
 			BaseDirectory = getExternalFilesDir(null).getAbsolutePath();
-			if (nativeIsShippingBuild())
+			if (BuildConfiguration.equals("Shipping"))
 			{
 				BaseDirectory = getFilesDir().getAbsolutePath();
 			}
@@ -932,7 +1061,7 @@ public class GameActivity extends $${gameActivitySuperClass}$$ implements Surfac
 		Log.debug("Parsing commandline: " + line);
 
 		// parse the line
-		ArrayList<String> sets = ParseSegments(line, " ", "\"", "\"");
+		ArrayList<String> sets = ParseSegments(line, ' ', "\"", "\"");
 		for (String set : sets)
 		{
 			String argument = set.trim();
@@ -940,7 +1069,7 @@ public class GameActivity extends $${gameActivitySuperClass}$$ implements Surfac
 			if (argument.startsWith("-"))
 			{
 				argument = argument.substring(1);
-				ArrayList<String> keyvalues = ParseSegments(argument, "=", "\"", "\"");
+				ArrayList<String> keyvalues = ParseSegments(argument, '=', "\"", "\"");
 				if (keyvalues.size() == 2)
 				{
 					CommandLineFlags.put(keyvalues.get(0).toLowerCase(), RemoveSurrounds(keyvalues.get(1), "\"", "\""));
@@ -1365,6 +1494,15 @@ public class GameActivity extends $${gameActivitySuperClass}$$ implements Surfac
 
 	private void getGraphicsInformation()
     {
+/* Example test values to skip data collection
+		EGLErrorStatus = "";
+		OpenGLVersion = "3.2";
+		TextureFormats = "ASTC,DXT,ATC,ETC2,ETC1";
+		bSupportsFloatingPointRenderTargets = true;
+		OpenGLVendor = "Qualcomm";
+		OpenGLDevice = "Adreno (TM) 740";
+		OpenGLDriver = "OpenGL ES 3.2 [REDACTED_EMAIL] (GIT@7b26bdd942, Iab69c31769, 1693222181) (Date:08/28/23)";
+*/
 		int initialError = EGL14.eglGetError();
 		EGLErrorStatus = (initialError == EGL14.EGL_SUCCESS) ? "" : "InitialEGLError = " + initialError + " ";
 
@@ -1444,7 +1582,7 @@ public class GameActivity extends $${gameActivitySuperClass}$$ implements Surfac
 		if (index >= 0)
 		{
 			OpenGLVersion = OpenGLDriver.substring(index + 10);
-			index = OpenGLVersion.indexOf(" ");
+			index = OpenGLVersion.indexOf(' ');
 			if (index > 0)
 			{
 				OpenGLVersion = OpenGLVersion.substring(0, index);
@@ -1472,7 +1610,7 @@ public class GameActivity extends $${gameActivitySuperClass}$$ implements Surfac
 		{
 			TextureFormats += "ATC,";
 		}
-		if (OpenGLVersion.charAt(0) >= '3')
+		if (OpenGLVersion.length() > 0 && OpenGLVersion.charAt(0) >= '3')
 		{
 			TextureFormats += "ETC2,";
 		}
@@ -1488,12 +1626,12 @@ public class GameActivity extends $${gameActivitySuperClass}$$ implements Surfac
 
 		if (inputLength >=2 && entryLength > 0 && entryLength == exitLength)
 		{
-			String start = input.substring(0,1);
-			String end = input.substring(inputLength-1);
+			char start = input.charAt(0);
+			char end = input.charAt(inputLength-1);
 
 			for (int index=0; index < entryLength; index++)
 			{
-				if (entry.substring(index, index+1).equals(start) && exit.substring(index, index+1).equals(end))
+				if (entry.charAt(index) == start && exit.charAt(index) == end)
 				{
 					return input.substring(1, inputLength-1);
 				}
@@ -1503,10 +1641,10 @@ public class GameActivity extends $${gameActivitySuperClass}$$ implements Surfac
 	}
 
 	// returns list of strings separated by split using entry/exit as pairing sets (ex. "(" and ")").  The entry and exit characters need to be in same order
-	private ArrayList<String> ParseSegments(String input, String split, String entry, String exit)
+	private ArrayList<String> ParseSegments(String input, char split, String entry, String exit)
 	{
 		ArrayList<String> output = new ArrayList<String>();
-		ArrayList<Integer> entryStack = new ArrayList<Integer>();
+		Stack<Integer> entryStack = new Stack<Integer>();
 
 		int startIndex = 0;
 		int scanIndex = 0;
@@ -1514,9 +1652,9 @@ public class GameActivity extends $${gameActivitySuperClass}$$ implements Surfac
 		int inputLength = input.length();
 		while (scanIndex < inputLength)
 		{
-			String scan = input.substring(scanIndex, scanIndex+1);
+			char scan = input.charAt(scanIndex);
 
-			if (scan.equals(split) && entryStack.size() == 0)
+			if (scan == split && entryStack.empty())
 			{
 				output.add(input.substring(startIndex, scanIndex).trim());
 				scanIndex++;
@@ -1525,17 +1663,16 @@ public class GameActivity extends $${gameActivitySuperClass}$$ implements Surfac
 			}
 			scanIndex++;
 
-			if (scan.equals("\\"))
+			if (scan == '\\')
 			{
 				scanIndex++;
 				continue;
 			}
 
-			if (entryStack.size() > 0 && exit.indexOf(scan) == exitIndex)
+			if (!entryStack.empty() && exit.indexOf(scan) == exitIndex)
 			{
-				int stackLength = entryStack.size() - 1;
-				entryStack.remove(stackLength);
-				exitIndex = stackLength > 0 ? entryStack.get(stackLength - 1) : -1;
+				entryStack.pop();
+				exitIndex = entryStack.empty() ? -1 : entryStack.peek();
 				continue;
 			}
 
@@ -1594,10 +1731,10 @@ public class GameActivity extends $${gameActivitySuperClass}$$ implements Surfac
 			String MatchString = "";
 
 			// deal with condition group (src,cmp,match)
-			ArrayList<String> groups = ParseSegments(RemoveSurrounds(condition, "(", ")"), ",", "\"", "\"");
+			ArrayList<String> groups = ParseSegments(RemoveSurrounds(condition, "(", ")"), ',', "\"", "\"");
 			for (String group : groups)
 			{
-				ArrayList<String> keyvalue = ParseSegments(group, "=", "\"", "\"");
+				ArrayList<String> keyvalue = ParseSegments(group, '=', "\"", "\"");
 				if (keyvalue.size() == 2)
 				{
 					String key = RemoveSurrounds(keyvalue.get(0), "\"", "\"");
@@ -1818,12 +1955,9 @@ public class GameActivity extends $${gameActivitySuperClass}$$ implements Surfac
 		return bConditionTrue;
 	}
 
-	protected void processSystemInfo()
+	protected void processSystemInfo(Lock lock)
 	{
-		SystemInfoLogs = new LinkedList<String>();
-
-		// check OpenGL information
-		getGraphicsInformation();
+		SystemInfo_InitLog();
 
 		// get CPU / chipset information
 		CpuInfo = getCpuInfo();
@@ -1948,8 +2082,6 @@ public class GameActivity extends $${gameActivitySuperClass}$$ implements Surfac
 		variables.put("littleCoreMask", "0x0");
 		variables.put("bigCoreMask", "0xffff");
 
-		variables.put("SRC_GpuFamily", OpenGLDevice);
-		variables.put("SRC_GlVersion", OpenGLDriver);
 		variables.put("SRC_AndroidVersion", android.os.Build.VERSION.RELEASE);
 		variables.put("SRC_DeviceMake", android.os.Build.MANUFACTURER);
 		variables.put("SRC_DeviceModel", android.os.Build.MODEL);
@@ -1959,10 +2091,7 @@ public class GameActivity extends $${gameActivitySuperClass}$$ implements Surfac
 		variables.put("SRC_UsingHoudini", bHasHoudini ? "true" : "false");
 		variables.put("SRC_HMDSystemName", "");
 
-		variables.put("SRC_GpuVendor", OpenGLVendor);
 		variables.put("SRC_SDKLevel", Integer.toString(android.os.Build.VERSION.SDK_INT));
-		variables.put("supportsFloatingPointRenderTargets", bSupportsFloatingPointRenderTargets ? "true" : "false");
-		variables.put("TextureFormats", TextureFormats);
 		variables.put("SRC_BuildConfig", BuildConfiguration);
 
 		// get the status bar and navigation bar heights (only works for 5.0+)
@@ -1997,6 +2126,26 @@ public class GameActivity extends $${gameActivitySuperClass}$$ implements Surfac
 		variables.put("screenWidth", Integer.toString(displaySize.x));
 		variables.put("screenHeight", Integer.toString(displaySize.y));
 
+		// check OpenGL information
+		if (lock != null)
+		{
+Trace.beginSection("GLLock");
+			lock.lock();
+Trace.endSection();
+		}
+Trace.beginSection("getGraphicsInfo");
+		getGraphicsInformation();
+		variables.put("SRC_GpuFamily", OpenGLDevice);
+		variables.put("SRC_GlVersion", OpenGLDriver);
+		variables.put("SRC_GpuVendor", OpenGLVendor);
+		variables.put("supportsFloatingPointRenderTargets", bSupportsFloatingPointRenderTargets ? "true" : "false");
+		variables.put("TextureFormats", TextureFormats);
+Trace.endSection();
+		if (lock != null)
+		{
+			lock.unlock();
+		}
+
 		// calculate a CRC to use for cache invalidation
 		CRC32 crc32 = new CRC32();
 		for (Map.Entry<String, String> entry : variables.entrySet())
@@ -2007,18 +2156,8 @@ public class GameActivity extends $${gameActivitySuperClass}$$ implements Surfac
 		ConfigRulesVarsCRC = crc32.getValue();
 	}
 
-	private boolean runConfigRules(Map<String, String> variables, String ProjectName, String PackageName)
+	protected boolean runConfigRules(Map<String, String> variables, String ProjectName, String PackageName)
 	{
-		// dump any SystemInfo logs
-		if (SystemInfoLogs != null)
-		{
-			for (String logline : SystemInfoLogs)
-			{
-				Log.debug(logline);
-			}
-			SystemInfoLogs = null;
-		}
-
 		// set early restart if -noearlyrestart not on commandline (configrules can clear this to disable or set to false)
 		if (!CommandLineFlags.containsKey("noearlyrestart"))
 		{
@@ -2029,7 +2168,8 @@ public class GameActivity extends $${gameActivitySuperClass}$$ implements Surfac
 		String SaveWarning = "";
 
 		// commandline flag may be used to skip dialog if not a shipping build
-		boolean bIgnoreDialog = !nativeIsShippingBuild() && CommandLineFlags.containsKey("noconfigrulesdialog");
+		boolean bIgnoreDialog = !BuildConfiguration.equals("Shipping")  // !nativeIsShippingBuild()
+								&& CommandLineFlags.containsKey("noconfigrulesdialog");
 
 		// NOTE: the "earlyrestart" variable and anything added to variables by <gameActivityPreConfigRulesParseAdditions> not included in ConfigRulesVarsCRC to invalidate cache!
 //$${gameActivityPreConfigRulesParseAdditions}$$
@@ -2074,7 +2214,7 @@ public class GameActivity extends $${gameActivitySuperClass}$$ implements Surfac
 			if (cacheFileContents != null)
 			{
 				String currentHeader = String.format("%d, %08X, %08X", ConfigRulesVersion, ConfigRulesCRC, ConfigRulesVarsCRC);
-				Log.debug("checking against: " + currentHeader);
+				SystemInfo_Log("checking against: " + currentHeader);
 
 				int headerSize = 4 + 8 + 8 + 4;
 				int sizeEncrypted = cacheFileContents.length - headerSize;
@@ -2085,7 +2225,7 @@ public class GameActivity extends $${gameActivitySuperClass}$$ implements Surfac
 					long cacheRulesCRC = buffer.getLong();
 					long cacheVarsCRC = buffer.getLong();
 					int cacheFlags = buffer.getInt();
-					Log.debug("read: " + String.format("%d, %08X, %08X", cacheVersion, cacheRulesCRC, cacheVarsCRC));
+					SystemInfo_Log("read: " + String.format("%d, %08X, %08X", cacheVersion, cacheRulesCRC, cacheVarsCRC));
 					if (cacheVersion == ConfigRulesVersion && cacheRulesCRC == ConfigRulesCRC && cacheVarsCRC == ConfigRulesVarsCRC)
 					{
 						byte[] decrypted = null;
@@ -2099,7 +2239,7 @@ public class GameActivity extends $${gameActivitySuperClass}$$ implements Surfac
 							}
 							catch (Exception e)
 							{
-								Log.debug("$$$ decrypt failure: " + e.toString());
+								SystemInfo_Log("configrules cache decrypt failure: " + e.toString());
 							}
 						}
 						else
@@ -2115,23 +2255,15 @@ public class GameActivity extends $${gameActivitySuperClass}$$ implements Surfac
 
 								variables = (HashMap<String, String>)objStream.readObject();
 								ConfigRulesVars = variables;
-								Log.debug("using configrules cache");
+								SystemInfo_Log("using configrules cache");
 
 								if ((cacheFlags & 2) > 0 && CpuInfo != null)
 								{
-									Log.debug("cpuInfo:");
-									for (Map.Entry<String, String> entry : CpuInfo.entrySet())
-									{
-										Log.debug("  " + entry.getKey() + " = " + entry.getValue());
-									}
+									SystemInfo_LogMap("cpuInfo:", CpuInfo);
 								}
 								if ((cacheFlags & 1) > 0)
 								{
-									Log.debug("variables:");
-									for (Map.Entry<String, String> entry : variables.entrySet())
-									{
-										Log.debug("  " + entry.getKey() + " = " + entry.getValue());
-									}
+									SystemInfo_LogMap("variables:", variables);
 								}
 
 								reader = null;
@@ -2142,7 +2274,7 @@ public class GameActivity extends $${gameActivitySuperClass}$$ implements Surfac
 							}
 							catch (Exception e)
 							{
-								Log.debug("$$$ deserialize failure: " + e.toString());
+								SystemInfo_Log("configrules cache deserialize failure: " + e.toString());
 							}
 						}
 					}
@@ -2176,13 +2308,15 @@ public class GameActivity extends $${gameActivitySuperClass}$$ implements Surfac
 						{
 							configRulesVersion = Integer.parseInt(line.substring(11));
 							variables.put("configRulesVersion", Integer.toString(configRulesVersion));
-							Log.debug("ConfigRules version: " + configRulesVersion);
+							SystemInfo_Log("ConfigRules version: " + configRulesVersion);
 						}
 						continue;
 					}
 
+					boolean bChangedVars = false;
+
 					// look for command
-					int index = line.indexOf(":");
+					int index = line.indexOf(':');
 					if (index == -1)
 					{
 						continue;
@@ -2196,7 +2330,7 @@ public class GameActivity extends $${gameActivitySuperClass}$$ implements Surfac
 						case CONFIGRULES_STATE_RUN:
 							if (command.equals("else") || command.equals("elseif") || command.equals("endif"))
 							{
-								Log.debug("ConfigRules: unexpected " + command + " encountered!");
+								SystemInfo_Log("ConfigRules: unexpected " + command + " encountered!");
 							}
 							break;
 
@@ -2241,7 +2375,7 @@ public class GameActivity extends $${gameActivitySuperClass}$$ implements Surfac
 							{
 								currentState = CONFIGRULES_STATE_FINDEND;
 
-								ArrayList<String> conditions = ParseSegments(line, ",", "(\"", ")\"");
+								ArrayList<String> conditions = ParseSegments(line, ',', "(\"", ")\"");
 								int setsize = conditions.size();
 								if (setsize > 0)
 								{
@@ -2259,7 +2393,7 @@ public class GameActivity extends $${gameActivitySuperClass}$$ implements Surfac
 							}
 							if (command.equals("else") || command.equals("elseif"))
 							{
-								Log.debug("ConfigRules: unexpected " + command + " while handling false condition!");
+								SystemInfo_Log("ConfigRules: unexpected " + command + " while handling false condition!");
 								continue;
 							}
 							break;
@@ -2285,7 +2419,7 @@ public class GameActivity extends $${gameActivitySuperClass}$$ implements Surfac
 							continue;
 
 						default:
-							Log.debug("ConfigRules: unknown state!");
+							SystemInfo_Log("ConfigRules: unknown state!");
 							continue;
 					}
 
@@ -2293,10 +2427,10 @@ public class GameActivity extends $${gameActivitySuperClass}$$ implements Surfac
 					if (command.equals("set"))
 					{
 						// set:(a=b[,c=d,...])
-						ArrayList<String> sets = ParseSegments(RemoveSurrounds(line, "(", ")"), ",", "(\"", ")\"");
+						ArrayList<String> sets = ParseSegments(RemoveSurrounds(line, "(", ")"), ',', "(\"", ")\"");
 						for (String assignment : sets)
 						{
-							ArrayList<String> keyvalue = ParseSegments(assignment, "=", "\"", "\"");
+							ArrayList<String> keyvalue = ParseSegments(assignment, '=', "\"", "\"");
 							if (keyvalue.size() == 2)
 							{
 								String key = RemoveSurrounds(keyvalue.get(0), "\"", "\"");
@@ -2316,15 +2450,17 @@ public class GameActivity extends $${gameActivitySuperClass}$$ implements Surfac
 								}
 							}
 						}
+						bChangedVars = true;
 					}
 					else if (command.equals("clear"))
 					{
 						// clear:(a[,b,...])
-						ArrayList<String> sets = ParseSegments(RemoveSurrounds(line, "(", ")"), ",", "(\"", ")\"");
+						ArrayList<String> sets = ParseSegments(RemoveSurrounds(line, "(", ")"), ',', "(\"", ")\"");
 						for (String key : sets)
 						{
 							variables.remove(RemoveSurrounds(key, "\"", "\""));
 						}
+						bChangedVars = true;
 					}
 					else if (command.equals("chipset"))
 					{
@@ -2333,7 +2469,7 @@ public class GameActivity extends $${gameActivitySuperClass}$$ implements Surfac
 						// only check if haven't already set chipset
 						if (!variables.containsKey("chipset"))
 						{
-							ArrayList<String> values = ParseSegments(line, ",", "\"", "\"");
+							ArrayList<String> values = ParseSegments(line, ',', "\"", "\"");
 							if (values.size() == 7)
 							{
 								String key = RemoveSurrounds(values.get(0), "\"", "\"").toLowerCase();
@@ -2346,8 +2482,9 @@ public class GameActivity extends $${gameActivitySuperClass}$$ implements Surfac
 									variables.put("bigCoreMask", RemoveSurrounds(values.get(5), "\"", "\""));
 									variables.put("littleCoreMask", RemoveSurrounds(values.get(6), "\"", "\""));
 
-									Log.debug("Chipset: " + variables.get("chipset"));
-									Log.debug("GPU: " + variables.get("gpu"));
+									SystemInfo_Log("Chipset: " + variables.get("chipset"));
+									SystemInfo_Log("GPU: " + variables.get("gpu"));
+									bChangedVars = true;
 								}
 							}
 						}
@@ -2364,7 +2501,7 @@ public class GameActivity extends $${gameActivitySuperClass}$$ implements Surfac
 						StateStack.push(currentState);
 						currentState = CONFIGRULES_STATE_FINDEND;
 
-						ArrayList<String> conditions = ParseSegments(line, ",", "(\"", ")\"");
+						ArrayList<String> conditions = ParseSegments(line, ',', "(\"", ")\"");
 						int setsize = conditions.size();
 						if (setsize > 0)
 						{
@@ -2376,13 +2513,13 @@ public class GameActivity extends $${gameActivitySuperClass}$$ implements Surfac
 					{
 						// condition:((SourceType=SRC_DeviceMake,CompareType=CMP_Equal,MatchString="samsung")),(SourceType=,CompareType=,MatchString=),...]),(a=b[,c=d,...]),(a[,b,...])
 						// if all the conditions are true, execute the optional sets and/or clears
-						ArrayList<String> conditionAndSets = ParseSegments(line, ",", "(\"", ")\"");
+						ArrayList<String> conditionAndSets = ParseSegments(line, ',', "(\"", ")\"");
 						int setsize = conditionAndSets.size();
 						if (setsize == 2 || setsize == 3)
 						{
-							ArrayList<String> conditions = ParseSegments(RemoveSurrounds(conditionAndSets.get(0), "(", ")"), ",", "(\"", ")\"");
-							ArrayList<String> sets = ParseSegments(RemoveSurrounds(conditionAndSets.get(1), "(", ")"), ",", "(\"", ")\"");
-							ArrayList<String> clears = (setsize == 3) ? ParseSegments(RemoveSurrounds(conditionAndSets.get(2), "(", ")"), ",", "(\"", ")\"") : new ArrayList<String>();
+							ArrayList<String> conditions = ParseSegments(RemoveSurrounds(conditionAndSets.get(0), "(", ")"), ',', "(\"", ")\"");
+							ArrayList<String> sets = ParseSegments(RemoveSurrounds(conditionAndSets.get(1), "(", ")"), ',', "(\"", ")\"");
+							ArrayList<String> clears = (setsize == 3) ? ParseSegments(RemoveSurrounds(conditionAndSets.get(2), "(", ")"), ',', "(\"", ")\"") : new ArrayList<String>();
 
 							boolean bConditionTrue = EvaluateConditions(variables, conditions, "");
 
@@ -2391,7 +2528,7 @@ public class GameActivity extends $${gameActivitySuperClass}$$ implements Surfac
 								// run the sets
 								for (String assignment : sets)
 								{
-									ArrayList<String> keyvalue = ParseSegments(assignment, "=", "\"", "\"");
+									ArrayList<String> keyvalue = ParseSegments(assignment, '=', "\"", "\"");
 									if (keyvalue.size() == 2)
 									{
 										String key = RemoveSurrounds(keyvalue.get(0), "\"", "\"");
@@ -2417,14 +2554,21 @@ public class GameActivity extends $${gameActivitySuperClass}$$ implements Surfac
 								{
 									variables.remove(RemoveSurrounds(key, "\"", "\""));
 								}
+
+								bChangedVars = true;
 							}
 						}
 					}
 
+					if (!bChangedVars)
+					{
+						continue;
+					}
+
 					// see if log message requested
 					if (variables.containsKey("log"))
 					{
-						Log.debug(variables.get("log"));
+						SystemInfo_Log(variables.get("log"));
 						variables.remove("log");
 					}
 
@@ -2435,11 +2579,7 @@ public class GameActivity extends $${gameActivitySuperClass}$$ implements Surfac
 						variables.remove("dumpcpuinfo");
 						if (CpuInfo != null)
 						{
-							Log.debug("cpuInfo:");
-							for (Map.Entry<String, String> entry : CpuInfo.entrySet())
-							{
-								Log.debug("  " + entry.getKey() + " = " + entry.getValue());
-							}
+							SystemInfo_LogMap("cpuInfo:", CpuInfo);
 						}
 					}
 
@@ -2448,11 +2588,7 @@ public class GameActivity extends $${gameActivitySuperClass}$$ implements Surfac
 					{
 						configFlags |= 1;
 						variables.remove("dumpvars");
-						Log.debug("variables:");
-						for (Map.Entry<String, String> entry : variables.entrySet())
-						{
-							Log.debug("  " + entry.getKey() + " = " + entry.getValue());
-						}
+						SystemInfo_LogMap("variables:", variables);
 					}
 
 					// if there was a raised error, stop
@@ -2544,7 +2680,7 @@ public class GameActivity extends $${gameActivitySuperClass}$$ implements Surfac
 			}
 			catch (IOException ie)
 			{
-				Log.debug("failed to read configuration rules: " + ie);
+				SystemInfo_Log("failed to read configuration rules: " + ie);
 				ConfigRulesStatus += "failed to read configuration rules : " + ie+ " ";
 			}
 			if (reader != null)
@@ -2652,7 +2788,7 @@ public class GameActivity extends $${gameActivitySuperClass}$$ implements Surfac
 		} 
 		catch (Exception e)
 		{
-			Log.debug("Failed to read memstats: MemstatLoggingInterval");
+			SystemInfo_Log("Failed to read memstats: MemstatLoggingInterval");
 		}
 		try 
 		{
@@ -2663,7 +2799,7 @@ public class GameActivity extends $${gameActivitySuperClass}$$ implements Surfac
 		} 
 		catch (Exception e)
 		{
-			Log.debug("Failed to read memstats: MemstatRSSRequiredDiff");
+			SystemInfo_Log("Failed to read memstats: MemstatRSSRequiredDiff");
 		}
 		
 
@@ -2672,6 +2808,17 @@ public class GameActivity extends $${gameActivitySuperClass}$$ implements Surfac
 		{
 			bSupportedDevice = variables.get("unsupportedDevice").equals("false");
 		}
+
+//$${gameActivityPostConfigRulesAdditions}$$
+//$${gameActivityFinalizeConfigRulesAdditions}$$
+
+		return !bForceGameEndWithError;
+	}
+
+	private void applyConfigRules(Map<String, String> variables, String ProjectName, String PackageName)
+	{
+		SystemInfo_DumpLog();
+
 		boolean bUseAffinity = variables.get("useAffinity").equals("true");
 		int littleCoreMask = Integer.decode(variables.get("littleCoreMask"));
 		int bigCoreMask = Integer.decode(variables.get("bigCoreMask"));
@@ -2694,11 +2841,6 @@ public class GameActivity extends $${gameActivitySuperClass}$$ implements Surfac
 			keyValues[keyValueIndex++] = entry.getValue();
 		}
 		nativeSetConfigRulesVariables(keyValues);
-
-//$${gameActivityPostConfigRulesAdditions}$$
-//$${gameActivityFinalizeConfigRulesAdditions}$$
-
-		return !bForceGameEndWithError;
 	}
 
 	public void LoggerCallback(String Level, String Tag, String Message)
@@ -2774,9 +2916,13 @@ public class GameActivity extends $${gameActivitySuperClass}$$ implements Surfac
 	@Override
 	public void onCreate(Bundle savedInstanceState)
 	{
+Trace.beginSection("onCreate Start");
 		Log.debug("LIFECYCLE: GameActivity.onCreate() called!, this=" + this);
 		
-		_extrasBundle = getIntent().getExtras();
+		if (!isStandalone && _extrasBundle == null)
+		{
+			_extrasBundle = getIntent().getExtras();
+		}
 
 		activityContext = getActivityOfContext(this);
 		assert activityContext != null;
@@ -2787,17 +2933,57 @@ public class GameActivity extends $${gameActivitySuperClass}$$ implements Surfac
 		AssetManagerReference = this.getAssets();
 		InternalFilesDir = getFilesDir().getAbsolutePath() + "/";
 		ExternalFilesDir = Objects.requireNonNull(getExternalFilesDir(null)).getAbsolutePath() + "/";
+Trace.endSection();
 
+		MainThreadId = android.os.Process.myTid();
+		MainThreadPriority = android.os.Process.getThreadPriority(MainThreadId);
+
+		// block getGraphicsInformation on ProcessSystemInfoThread until GLES loads
+		ProcessSystemInfoLock.lock();
+
+//$${gameActivityOnCreateBeforeProcessSystemInfoAdditions}$$
 		ProcessSystemInfoThread = new Thread(new Runnable() {
-            public void run() {
-				processSystemInfo();
+			public void run() {
+				android.os.Process.setThreadPriority(android.os.Process.myTid(), MainThreadPriority);
+				// parse the commandline
+				ParseCommandline(ProjectName, UseExternalFilesDir);
+Trace.beginSection("processSystemInfo");
+				// collect system info
+				processSystemInfo(ProcessSystemInfoLock);
+Trace.endSection();
+
+				// run config rules
+//$${gameActivityProcessSystemInfoThreadBeforeConfigRulesAdditions}$$
+Trace.beginSection("runConfigRules");
+				Log.debug("$$$ configrules: run config rules");
+				ConfigRulesRunResult = runConfigRules(ConfigRulesVars, ProjectName, appPackageName);
+				Log.debug("$$$ configrules: run completed");
+Trace.endSection();
+//$${gameActivityProcessSystemInfoThreadAfterConfigRulesAdditions}$$
 			}
 		});
 		ProcessSystemInfoThread.setName("ProcessSysInfo");
 		ProcessSystemInfoThread.start();
 
+		// load GLES/EGL now so available for getGraphicsInformation in ProcessSystemInfo thread without RenderDoc hang
+		// cannot be in another thread as it will just block (fairly quick to load so not a big deal to do it here)
+Trace.beginSection("loadGLES");
+		try
+		{
+			System.loadLibrary("GLESv3");
+			System.loadLibrary("EGL");
+		}
+		catch (Exception e)
+		{
+		}
+		// unlock so getGraphicsInformation is allowed to run
+		ProcessSystemInfoLock.unlock();
+Trace.endSection();
+
 //$${gameActivityOnCreateBeforeLoadLibrariesAdditions}$$
+Trace.beginSection("loadLibraries");
 		LoadLibraries();
+Trace.endSection();
 
 		// Bug with certain devices calling onCreate twice, even with singleTask set and System.exit called in onDestroy.
 		if (nativeIsShippingBuild())
@@ -2827,7 +3013,7 @@ public class GameActivity extends $${gameActivitySuperClass}$$ implements Surfac
 	protected void onCreatePreBody()
 	{
 		// Retrieve version code, name, and target SDK
-		try
+		try 
 		{
 			PackageInfo packageInfo = getPackageManager().getPackageInfo(getPackageName(), 0);
 			VersionCode = packageInfo.versionCode;
@@ -2894,6 +3080,16 @@ public class GameActivity extends $${gameActivitySuperClass}$$ implements Surfac
 				Log.debug( "BuildConfiguration not found" );
 			}
 
+			if (bundle.containsKey("com.epicgames.unreal.GameActivity.bUseExternalFilesDir"))
+            {
+                UseExternalFilesDir = bundle.getBoolean("com.epicgames.unreal.GameActivity.bUseExternalFilesDir");
+                Log.debug( "UseExternalFilesDir set to " + UseExternalFilesDir);
+            }
+            else
+            {
+                Log.debug( "bUseExternalFilesDir not found. Leaving as " + UseExternalFilesDir);
+            }
+
 			if (bundle.containsKey("com.epicgames.unreal.GameActivity.bSupportsVulkan"))
 			{
 				bSupportsVulkan = bundle.getBoolean("com.epicgames.unreal.GameActivity.bSupportsVulkan");
@@ -2918,6 +3114,23 @@ public class GameActivity extends $${gameActivitySuperClass}$$ implements Surfac
 		{
 			Log.debug( "Failed to load meta-data: NullPointer: " + e.getMessage());
 		}
+
+		if (!isStandalone && _extrasBundle != null)
+		{
+			ShouldHideUI = _extrasBundle.getString("ShouldHideUI") != null;
+			UseDisplayCutout = _extrasBundle.getString("UseDisplayCutout") != null;
+
+			String CmdLine = _extrasBundle.getString("cmdline");
+			if (CmdLine != null && !BuildConfiguration.equals("Shipping"))
+			{
+				CommandLineFull = CmdLine;
+			}
+		}
+
+		if (isStandalone)
+		{
+			ParseCommandline(ProjectName, UseExternalFilesDir);
+		}
 	}
 
 	protected void onCreateBody(Bundle savedInstanceState)
@@ -2999,15 +3212,8 @@ public class GameActivity extends $${gameActivitySuperClass}$$ implements Surfac
 		boolean bIsShipping = nativeIsShippingBuild();
 
 		// create splashscreen dialog (if launched by SplashActivity)
-		if (!isStandalone && _extrasBundle == null)
-		{
-			_extrasBundle = getIntent().getExtras();
-		}
-
 		if (!isStandalone && _extrasBundle != null)
 		{
-			ShouldHideUI = _extrasBundle.getString("ShouldHideUI") != null;
-			UseDisplayCutout = _extrasBundle.getString("UseDisplayCutout") != null;
 			if (_extrasBundle.getString("UseSplashScreen") != null)
 			{
 				SplashScreenLaunch = true;
@@ -3052,12 +3258,6 @@ public class GameActivity extends $${gameActivitySuperClass}$$ implements Surfac
 				e.printStackTrace();
 			}
 
-			String CmdLine = _extrasBundle.getString("cmdline");
-			if (CmdLine != null && !bIsShipping)
-			{
-				CommandLineFull = CmdLine;
-			}
-
 			if (!isStandalone && UseDisplayCutout)
 			{
 				// will not be true if not Android Pie or later
@@ -3121,8 +3321,8 @@ public class GameActivity extends $${gameActivitySuperClass}$$ implements Surfac
 				}
 
 				String[] statusLines = FindLineFromFile("/proc/self/status",new String[] {"VmSize:","VmHWM:", "VmRSS:", "VmData:", "VmSwap:"});
-				String rssSubString = statusLines[2].substring(0, statusLines[2].lastIndexOf(" "));
-				int rssMemory = Integer.parseInt(rssSubString.substring(rssSubString.lastIndexOf(" ")+1, rssSubString.length()));
+				String rssSubString = statusLines[2].substring(0, statusLines[2].lastIndexOf(' '));
+				int rssMemory = Integer.parseInt(rssSubString.substring(rssSubString.lastIndexOf(' ')+1, rssSubString.length()));
 				int rssDiffed = Math.abs(rssMemory - LastRSSMemoryValue);
 
 				String ProcessMemInfoLine = "Used memory: ";
@@ -3248,7 +3448,7 @@ public class GameActivity extends $${gameActivitySuperClass}$$ implements Surfac
 						if (index >= 0)
 						{
 							dump = dump.substring(index+2);
-							index = dump.indexOf(" ");
+							index = dump.indexOf(' ');
 							if (index >= 0)
 							{
 								VulkanLevel = Integer.parseInt(dump.substring(0, index));
@@ -3265,7 +3465,7 @@ public class GameActivity extends $${gameActivitySuperClass}$$ implements Surfac
 						if (index >= 0)
 						{
 							dump = dump.substring(index+2);
-							index = dump.indexOf(" ");
+							index = dump.indexOf(' ');
 							if (index >= 0)
 							{
 								VulkanVersion = Integer.parseInt(dump.substring(0, index));
@@ -3406,16 +3606,6 @@ public class GameActivity extends $${gameActivitySuperClass}$$ implements Surfac
 				}
 			}
 
-			if (bundle.containsKey("com.epicgames.unreal.GameActivity.bUseExternalFilesDir"))
-            {
-                UseExternalFilesDir = bundle.getBoolean("com.epicgames.unreal.GameActivity.bUseExternalFilesDir");
-                Log.debug( "UseExternalFilesDir set to " + UseExternalFilesDir);
-            }
-            else
-            {
-                Log.debug( "bUseExternalFilesDir not found. Leaving as " + UseExternalFilesDir);
-            }
-
 			if (bundle.containsKey("com.epicgames.unreal.GameActivity.bPublicLogFiles"))
             {
                 PublicLogFiles = bundle.getBoolean("com.epicgames.unreal.GameActivity.bPublicLogFiles");
@@ -3447,10 +3637,6 @@ public class GameActivity extends $${gameActivitySuperClass}$$ implements Surfac
 			Log.debug( "Failed to load meta-data: " + e.getMessage());
 		}
 
-		// parse the commandline
-		ParseCommandline(ProjectName, UseExternalFilesDir);
-
-//$${gameActivityBeforeConfigRulesAppliedAdditions}$$
 
 		// run configuration rules (wait for SystemInfoLogs fetch to complete first)
 		try
@@ -3460,12 +3646,14 @@ public class GameActivity extends $${gameActivitySuperClass}$$ implements Surfac
 		catch (Exception e)
 		{
 		}
-		Log.debug("$$$ configrules: run config rules");
-		if (!runConfigRules(ConfigRulesVars, ProjectName, appPackageName))
+//$${gameActivityBeforeConfigRulesAppliedAdditions}$$
+		if (!ConfigRulesRunResult)
 		{
 			// configuration failed (show error)
 			return;
 		}
+		Log.debug("$$$ configrules: apply config rules");
+		applyConfigRules(ConfigRulesVars, ProjectName, appPackageName);
 		Log.debug("$$$ configrules: finish config rules");
 
 		String APKPath = getPackageResourcePath();
@@ -3511,6 +3699,7 @@ public class GameActivity extends $${gameActivitySuperClass}$$ implements Surfac
 			this.setVolumeControlStream(AudioManager.STREAM_MUSIC);
 		}
 
+Trace.beginSection("UI Building");
 		AlertDialog.Builder builder;
 
 		consoleInputBox = new ConsoleKeyboardInput(getThisActivity());
@@ -3834,6 +4023,7 @@ public class GameActivity extends $${gameActivitySuperClass}$$ implements Surfac
 				}
 			});
 		}
+Trace.endSection();
 
 //$${gameActivityOnCreateAdditions}$$
 //$${gameActivityOnCreateFinalAdditions}$$
@@ -4943,7 +5133,7 @@ public class GameActivity extends $${gameActivitySuperClass}$$ implements Surfac
 		if (!URL.contains("://"))
 		{
 			// add http:// if there isn't a scheme before a colon
-			if (URL.indexOf(":") < 1)
+			if (URL.indexOf(':') < 1)
 			{
 				URL = "http://" + URL;
 				Log.debug("[JAVA} AndroidThunkJava_LaunchURL: corrected URL = " + URL);
@@ -7234,13 +7424,13 @@ public class GameActivity extends $${gameActivitySuperClass}$$ implements Surfac
 		String intentAction = "";
 		String intentComponent = "";
 
-		int intentActionIndex = packageName.indexOf(":");
+		int intentActionIndex = packageName.indexOf(':');
 		if (intentActionIndex >= 0)
 		{
 			intentAction = packageName.substring(intentActionIndex + 1);
 			packageName = packageName.substring(0, intentActionIndex);
 
-			intentActionIndex = intentAction.indexOf(":");
+			intentActionIndex = intentAction.indexOf(':');
 			if (intentActionIndex >= 0)
 			{
 				intentComponent = intentAction.substring(intentActionIndex + 1);
@@ -7306,13 +7496,13 @@ public class GameActivity extends $${gameActivitySuperClass}$$ implements Surfac
 		String intentAction = "";
 		String intentComponent = "";
 
-		int intentActionIndex = packageName.indexOf(":");
+		int intentActionIndex = packageName.indexOf(':');
 		if (intentActionIndex >= 0)
 		{
 			intentAction = packageName.substring(intentActionIndex + 1);
 			packageName = packageName.substring(0, intentActionIndex);
 
-			intentActionIndex = intentAction.indexOf(":");
+			intentActionIndex = intentAction.indexOf(':');
 			if (intentActionIndex >= 0)
 			{
 				intentComponent = intentAction.substring(intentActionIndex + 1);
diff --git a/Engine/Build/Android/Java/src/com/epicgames/unreal/SplashActivity.java b/Engine/Build/Android/Java/src/com/epicgames/unreal/SplashActivity.java
--- a/Engine/Build/Android/Java/src/com/epicgames/unreal/SplashActivity.java
+++ b/Engine/Build/Android/Java/src/com/epicgames/unreal/SplashActivity.java
@@ -2,7 +2,7 @@
 
 package com.epicgames.unreal;
 
-import java.util.ArrayList;
+import java.util.LinkedList;
 import java.util.List;
 
 import android.annotation.SuppressLint;
@@ -220,7 +220,7 @@ public class SplashActivity extends Activity
 		}
 
 		// make a list of ungranted dangerous permissions in manifest required at start of GameActivity and request any we still need
-		ArrayList<String> ungrantedPermissions = getUngrantedPermissions(this, getDangerousPermissions(pm, packageName), permissionsRequiredAtStart);
+		LinkedList<String> ungrantedPermissions = getUngrantedPermissions(this, getDangerousPermissions(pm, packageName), permissionsRequiredAtStart);
 		if (ungrantedPermissions.size() > 0)
 		{
 			ActivityCompat.requestPermissions(this, ungrantedPermissions.toArray(new String[ungrantedPermissions.size()]), PERMISSION_REQUEST_CODE);
@@ -255,7 +255,7 @@ public class SplashActivity extends Activity
 	private String[] filterRequiredPermissions(String permissions)
 	{
 		String manufacturer = android.os.Build.MANUFACTURER;
-		ArrayList<String> keptPermissions = new ArrayList<>();
+		LinkedList<String> keptPermissions = new LinkedList<>();
 		String[] requiredPermissions = permissions.split(",");
 		for (String required : requiredPermissions)
 		{
@@ -371,9 +371,9 @@ public class SplashActivity extends Activity
 		return keptPermissions.toArray(new String[keptPermissions.size()]);
 	}
 
-	public ArrayList<String> getUngrantedPermissions(Context context, ArrayList<String> dangerousPermissions, String[] requiredPermissions)
+	public LinkedList<String> getUngrantedPermissions(Context context, LinkedList<String> dangerousPermissions, String[] requiredPermissions)
 	{
-		ArrayList<String> ungrantedPermissions = new ArrayList<>();
+		LinkedList<String> ungrantedPermissions = new LinkedList<>();
 		if (dangerousPermissions.size() > 0)
 		{
 			for (String required : requiredPermissions)
@@ -390,10 +390,10 @@ public class SplashActivity extends Activity
 		return ungrantedPermissions;
 	}
 
-	public ArrayList<String> getDangerousPermissions(PackageManager pm, String packageName)
+	public LinkedList<String> getDangerousPermissions(PackageManager pm, String packageName)
 	{
 		int targetSdkVersion = 0;
-		ArrayList<String> dangerousPermissions = new ArrayList<>();
+		LinkedList<String> dangerousPermissions = new LinkedList<>();
 		try 
 		{
 			PackageInfo packageInfo = pm.getPackageInfo(packageName, 0);
@@ -491,7 +491,7 @@ public class SplashActivity extends Activity
 						@Override
 						public void onClick(DialogInterface dialog, int id) {
 							dialog.dismiss();
-							ArrayList<String> ungrantedPermissions = getUngrantedPermissions(SplashActivity.this, getDangerousPermissions(pm, packageName), permissionsRequiredAtStart);
+							LinkedList<String> ungrantedPermissions = getUngrantedPermissions(SplashActivity.this, getDangerousPermissions(pm, packageName), permissionsRequiredAtStart);
 							if (ungrantedPermissions.size() > 0)
 							{
 								ActivityCompat.requestPermissions(SplashActivity.this, ungrantedPermissions.toArray(new String[ungrantedPermissions.size()]), PERMISSION_REQUEST_CODE);
diff --git a/Engine/Source/Programs/UnrealBuildTool/Platform/Android/UEDeployAndroid.cs b/Engine/Source/Programs/UnrealBuildTool/Platform/Android/UEDeployAndroid.cs
--- a/Engine/Source/Programs/UnrealBuildTool/Platform/Android/UEDeployAndroid.cs
+++ b/Engine/Source/Programs/UnrealBuildTool/Platform/Android/UEDeployAndroid.cs
@@ -4150,8 +4150,18 @@ namespace UnrealBuildTool
 			GradleBuildAdditionsContent.AppendLine("\t}");
 			GradleBuildAdditionsContent.AppendLine("}");
 
+			String BuildGradleAdditionsDefault = "";
+			bool bJavaPerfettoMarkers = false;
+			Ini.GetBool("/Script/AndroidRuntimeSettings.AndroidRuntimeSettings", "bJavaPerfettoMarkers", out bJavaPerfettoMarkers);
+			if (bJavaPerfettoMarkers)
+			{
+				BuildGradleAdditionsDefault = @"
+dependencies.implementation('androidx.tracing:tracing:1.2.0')
+";
+			}
+
 			// Add any UPL app buildGradleAdditions
-			GradleBuildAdditionsContent.Append(UPL.ProcessPluginNode(NDKArch, "buildGradleAdditions", ""));
+			GradleBuildAdditionsContent.Append(UPL.ProcessPluginNode(NDKArch, "buildGradleAdditions", BuildGradleAdditionsDefault));
 
 			string GradleBuildAdditionsFilename = Path.Combine(UnrealBuildGradleAppPath, "buildAdditions.gradle");
 			File.WriteAllText(GradleBuildAdditionsFilename, GradleBuildAdditionsContent.ToString());
@@ -6158,6 +6168,13 @@ import java.util.Collection;
 					break;
 			}
 
+			bool bJavaPerfettoMarkers = false;
+			Ini.GetBool("/Script/AndroidRuntimeSettings.AndroidRuntimeSettings", "bJavaPerfettoMarkers", out bJavaPerfettoMarkers);
+			if (bJavaPerfettoMarkers)
+			{
+				GameActivityImportAdditionsDefault += "import androidx.tracing.Trace;\n";
+			}
+
 			Dictionary<string, string> Replacements = new Dictionary<string, string>{
 				{ "//$${gameActivityMemStatAdditions}$$", UPL.ProcessPluginNode(NDKArch, "gameActivityMemStatAdditions", "")},
 				{ "//$${gameActivityImportAdditions}$$", UPL.ProcessPluginNode(NDKArch, "gameActivityImportAdditions", GameActivityImportAdditionsDefault)},
@@ -6165,6 +6182,9 @@ import java.util.Collection;
 				{ "//$${gameActivityImplementsAdditions}$$", UPL.ProcessPluginNode(NDKArch, "gameActivityImplementsAdditions", "")},
 				{ "//$${gameActivityClassAdditions}$$", UPL.ProcessPluginNode(NDKArch, "gameActivityClassAdditions", GameActivityClassAdditionsDefault)},
 				{ "//$${gameActivityReadMetadataAdditions}$$", UPL.ProcessPluginNode(NDKArch, "gameActivityReadMetadataAdditions", "")},
+				{ "//$${gameActivityOnCreateBeforeProcessSystemInfoAdditions}$$", UPL.ProcessPluginNode(NDKArch, "gameActivityOnCreateBeforeProcessSystemInfoAdditions", "")},
+				{ "//$${gameActivityProcessSystemInfoThreadBeforeConfigRulesAdditions}$$", UPL.ProcessPluginNode(NDKArch, "gameActivityProcessSystemInfoThreadBeforeConfigRulesAdditions", "")},
+				{ "//$${gameActivityProcessSystemInfoThreadAfterConfigRulesAdditions}$$", UPL.ProcessPluginNode(NDKArch, "gameActivityProcessSystemInfoThreadAfterConfigRulesAdditions", "")},
 				{ "//$${gameActivityOnCreateBeforeLoadLibrariesAdditions}$$", UPL.ProcessPluginNode(NDKArch, "gameActivityOnCreateBeforeLoadLibrariesAdditions", "")},
 				{ "//$${gameActivityOnCreateBeginningAdditions}$$", UPL.ProcessPluginNode(NDKArch, "gameActivityOnCreateBeginningAdditions", "")},
 				{ "//$${gameActivityOnCreateAdditions}$$", UPL.ProcessPluginNode(NDKArch, "gameActivityOnCreateAdditions", "")},
@@ -6211,18 +6231,61 @@ import java.util.Collection;
 			{
 				string SrcLine = TemplateSrc[LineIndex];
 				bool Changed = false;
-				foreach (KeyValuePair<string, string> KVP in Replacements)
+
+				if (!bJavaPerfettoMarkers)
 				{
-					if (SrcLine.Contains(KVP.Key))
+					if (SrcLine.StartsWith("Trace."))
 					{
-						SrcLine = SrcLine.Replace(KVP.Key, KVP.Value);
+						TemplateSrc[LineIndex] = "";
 						Changed = true;
+						TemplateChanged = true;
 					}
 				}
-				if (Changed)
+
+				if (!Changed)
 				{
-					TemplateSrc[LineIndex] = SrcLine;
-					TemplateChanged = true;
+					foreach (KeyValuePair<string, string> KVP in Replacements)
+					{
+						if (SrcLine.Contains(KVP.Key))
+						{
+							if (bJavaPerfettoMarkers)
+							{
+								SrcLine = SrcLine.Replace(KVP.Key, KVP.Value);
+							}
+							else
+							{
+								int TraceStartIndex = KVP.Value.IndexOf("Trace.");
+								if (TraceStartIndex == 0 || (TraceStartIndex > 0 && KVP.Value[TraceStartIndex - 1] == '\n'))
+								{
+									String Addition = KVP.Value;
+									do
+									{
+										int TraceEndIndex = Addition.IndexOf("\n", TraceStartIndex);
+										if (TraceEndIndex < 0)
+										{
+											Addition = Addition.Substring(0, TraceStartIndex);
+										}
+										else
+										{
+											Addition = Addition.Substring(0, TraceStartIndex) + Addition.Substring(TraceEndIndex + 1);
+										}
+										TraceStartIndex = Addition.IndexOf("Trace.", TraceStartIndex);
+									} while (TraceStartIndex >= 0);
+									SrcLine = SrcLine.Replace(KVP.Key, Addition);
+								}
+								else
+								{
+									SrcLine = SrcLine.Replace(KVP.Key, KVP.Value);
+								}
+							}
+							Changed = true;
+						}
+					}
+					if (Changed)
+					{
+						TemplateSrc[LineIndex] = SrcLine;
+						TemplateChanged = true;
+					}
 				}
 			}
 
