Add preparsing to configrules for improved startup time
- ConfigRulesTool parser mode ('p' command)
- New -verbose flag for 'p' (regex optimizes) and 'd' commands (offsets for debug)
- added validation with errors showing line number
- RegEx optimizer (replaces simple string comparisons)
- commandline -crvars= option for non-Shipping builds
- Trace.beginSection/endSection filtering in GameApplication.java.template

#android

diff --git a/Engine/Build/Android/Java/src/com/epicgames/unreal/CRToken.java b/Engine/Build/Android/Java/src/com/epicgames/unreal/CRToken.java
new file mode 100644
--- /dev/null
+++ b/Engine/Build/Android/Java/src/com/epicgames/unreal/CRToken.java
@@ -0,0 +1,164 @@
+// Copyright Epic Games, Inc. All Rights Reserved.
+
+package com.epicgames.unreal;
+
+import java.util.HashMap;
+import java.util.Map;
+
+public class CRToken {
+	public static final int COMMAND_UNKNOWN = -1;
+	public static final int COMMAND_SET = 0;
+	public static final int COMMAND_CLEAR = 1;
+	public static final int COMMAND_CHIPSET = 2;
+	public static final int COMMAND_CONDITION = 3;
+	public static final int COMMAND_IF = 4;
+	public static final int COMMAND_ELSEIF = 5;
+	public static final int COMMAND_ELSE = 6;
+	public static final int COMMAND_ENDIF = 7;
+	public static final int COMMAND_GOTO = 8;
+	public static final int COMMAND_EOF = 99;
+
+	public static final int CONDKEY_UNKNOWN = -1;
+	public static final int CONDKEY_SOURCETYPE = 0;
+	public static final int CONDKEY_COMPARETYPE = 1;
+	public static final int CONDKEY_MATCHSTRING = 2;
+
+	public static final int SOURCE_LITERAL = 0;
+	public static final int SOURCE_EXIST = 1;
+	public static final int SOURCE_PREVREGEX = 2;
+	public static final int SOURCE_COMMANDLINE = 3;
+
+	public static final int COMPARE_IGNORE = 50;
+
+	public static final int COMPARE_UNKNOWN = -1;
+	public static final int COMPARE_EXIST = 0;
+	public static final int COMPARE_NOTEXIST = 1;
+	public static final int COMPARE_EQUAL = 2;
+	public static final int COMPARE_NOTEQUAL = 3;
+	public static final int COMPARE_REGEX = 4;
+	public static final int COMPARE_STARTS = 5;
+	public static final int COMPARE_CONTAINS = 6;
+	public static final int COMPARE_ENDS = 7;
+	public static final int COMPARE_VEQUAL = 8;
+	public static final int COMPARE_VNOTEQUAL = 9;
+	public static final int COMPARE_LESS = 10;
+	public static final int COMPARE_LESSEQUAL = 11;
+	public static final int COMPARE_GREATER = 12;
+	public static final int COMPARE_GREATEREQUAL = 13;
+
+	public static final int SETVARFLAGS_NONE = 0;
+	public static final int SETVARFLAGS_APPEND = 1;
+	public static final int SETVARFLAGS_EXPAND = 2;
+
+	private static HashMap<String,Integer> CommandTokens = new HashMap<>();
+	private static HashMap<String,Integer> CondKeyTokens = new HashMap<>();
+	private static HashMap<String,Integer> SourceTokens = new HashMap<>();
+	private static HashMap<String,Integer> CompareTokens = new HashMap<>();
+
+	private static HashMap<Integer,String> CommandStrings = new HashMap<>();
+	private static HashMap<Integer,String> CondKeyStrings = new HashMap<>();
+	private static HashMap<Integer,String> SourceStrings = new HashMap<>();
+	private static HashMap<Integer,String> CompareStrings = new HashMap<>();
+
+	public static void Init()
+	{
+		CommandTokens.put("set", COMMAND_SET);
+		CommandTokens.put("clear", COMMAND_CLEAR);
+		CommandTokens.put("chipset", COMMAND_CHIPSET);
+		CommandTokens.put("condition", COMMAND_CONDITION);
+		CommandTokens.put("if", COMMAND_IF);
+		CommandTokens.put("elseif", COMMAND_ELSEIF);
+		CommandTokens.put("else", COMMAND_ELSE);
+		CommandTokens.put("endif", COMMAND_ENDIF);
+
+		CondKeyTokens.put("SourceType", CONDKEY_SOURCETYPE);
+		CondKeyTokens.put("CompareType", CONDKEY_COMPARETYPE);
+		CondKeyTokens.put("MatchString", CONDKEY_MATCHSTRING);
+
+		SourceTokens.put("[EXIST]", SOURCE_EXIST);
+		SourceTokens.put("SRC_PreviousRegexMatch", SOURCE_PREVREGEX);
+		SourceTokens.put("SRC_CommandLine", SOURCE_COMMANDLINE);
+
+		CompareTokens.put("CMP_Exist", COMPARE_EXIST);
+		CompareTokens.put("CMP_NotExist", COMPARE_NOTEXIST);
+		CompareTokens.put("CMP_Equal", COMPARE_EQUAL);
+		CompareTokens.put("CMP_NotEqual", COMPARE_NOTEQUAL);
+		CompareTokens.put("CMP_EqualIgnore", COMPARE_EQUAL + COMPARE_IGNORE);
+		CompareTokens.put("CMP_NotEqualIgnore", COMPARE_NOTEQUAL + COMPARE_IGNORE);
+		CompareTokens.put("CMP_Regex", COMPARE_REGEX);
+		CompareTokens.put("CMP_Starts", COMPARE_STARTS);
+		CompareTokens.put("CMP_Contains", COMPARE_CONTAINS);
+		CompareTokens.put("CMP_Ends", COMPARE_ENDS);
+		CompareTokens.put("CMP_VEqual", COMPARE_VEQUAL);
+		CompareTokens.put("CMP_VNotEqual", COMPARE_VNOTEQUAL);
+		CompareTokens.put("CMP_Less", COMPARE_LESS);
+		CompareTokens.put("CMP_LessIgnore", COMPARE_LESS + COMPARE_IGNORE);
+		CompareTokens.put("CMP_LessEqual", COMPARE_LESSEQUAL);
+		CompareTokens.put("CMP_LessEqualIgnore", COMPARE_LESSEQUAL + COMPARE_IGNORE);
+		CompareTokens.put("CMP_Greater", COMPARE_GREATER);
+		CompareTokens.put("CMP_GreaterIgnore", COMPARE_GREATER + COMPARE_IGNORE);
+		CompareTokens.put("CMP_GreaterEqual", COMPARE_GREATEREQUAL);
+		CompareTokens.put("CMP_GreaterEqualIgnore", COMPARE_GREATEREQUAL + COMPARE_IGNORE);
+	}
+
+	public static void InitStrings()
+	{
+		for (Map.Entry<String, Integer> entry : CommandTokens.entrySet())
+		{
+			CommandStrings.put((Integer)entry.getValue(), (String)entry.getKey());
+		}
+		for (Map.Entry<String, Integer> entry : CondKeyTokens.entrySet())
+		{
+			CondKeyStrings.put((Integer)entry.getValue(), (String)entry.getKey());
+		}
+		for (Map.Entry<String, Integer> entry : SourceTokens.entrySet())
+		{
+			SourceStrings.put((Integer)entry.getValue(), (String)entry.getKey());
+		}
+		for (Map.Entry<String, Integer> entry : CompareTokens.entrySet())
+		{
+			CompareStrings.put((Integer)entry.getValue(), (String)entry.getKey());
+		}
+	}
+
+	public static int findCommand(String in)
+	{
+		return CommandTokens.getOrDefault(in, COMMAND_UNKNOWN);
+	}
+
+	public static int findCondKey(String in)
+	{
+		return CondKeyTokens.getOrDefault(in, CONDKEY_UNKNOWN);
+	}
+
+	public static int findSource(String in)
+	{
+		return SourceTokens.getOrDefault(in, SOURCE_LITERAL);
+	}
+
+	public static int findCompare(String in)
+	{
+		return CompareTokens.getOrDefault(in, COMPARE_UNKNOWN);
+	}
+
+	public static String getCommand(int token)
+	{
+		return CommandStrings.getOrDefault(token, "UNKNOWN");
+	}
+
+	public static String getCondKey(int token)
+	{
+		return CondKeyStrings.getOrDefault(token, "UNKNOWN");
+	}
+
+	public static String getSource(int token)
+	{
+		return SourceStrings.getOrDefault(token, "UNKNOWN");
+	}
+
+	public static String getCompare(int token)
+	{
+		return CompareStrings.getOrDefault(token, "UNKNOWN");
+	}
+}
+
diff --git a/Engine/Build/Android/Java/src/com/epicgames/unreal/GameActivity.java.template b/Engine/Build/Android/Java/src/com/epicgames/unreal/GameActivity.java.template
--- a/Engine/Build/Android/Java/src/com/epicgames/unreal/GameActivity.java.template
+++ b/Engine/Build/Android/Java/src/com/epicgames/unreal/GameActivity.java.template
@@ -198,6 +198,7 @@ import java.util.TimeZone;
 import java.util.UUID;
 import java.io.ByteArrayOutputStream;
 
+import com.epicgames.unreal.CRToken;
 import com.epicgames.unreal.GooglePlayLicensing;
 
 // Console commands listener, only for debug builds
@@ -370,6 +371,11 @@ public class GameActivity extends $${gameActivitySuperClass}$$ implements Surfac
 	private Lock SystemInfoLogsLock = new ReentrantLock();
 	private Lock ProcessSystemInfoLock = new ReentrantLock();
 	protected boolean ConfigRulesRunResult = true;
+	protected java.security.Key ConfigRulesKey = null;
+	protected java.security.Key ConfigRulesCacheKey = null;
+	protected byte[] ConfigRulesBytesAssets = null;
+	protected byte[] ConfigRulesBytesInternal = null;
+	protected byte[] ConfigRulesBytesExternal = null;
 	protected long ConfigRulesCRC = 0;
 	protected long ConfigRulesVarsCRC = 0;
 	protected Map<String, String> ConfigRulesVars = null;
@@ -1148,7 +1154,7 @@ public class GameActivity extends $${gameActivitySuperClass}$$ implements Surfac
 
 		ByteBuffer buffer = ByteBuffer.wrap(bytesToDecompress);
 		int signature = buffer.getShort();
-		if (signature != 0x39d8)
+		if (signature != 0x39d8 && signature != 0x39da)
 		{
 			Log.debug("configrulesversion: token ("+signature+")");
 			return -1;
@@ -1170,13 +1176,121 @@ public class GameActivity extends $${gameActivitySuperClass}$$ implements Surfac
 		return new SecretKeySpec(salt, "AES");
 	}
 
-	private InputStream getConfigRulesDecompressed(byte[] bytesToDecompress)
+	class ByteArrayInputStream2 extends InputStream
+	{
+		private byte[] data;
+		private int pos;
+		private int mark;
+
+		public ByteArrayInputStream2(byte[] data)
+		{
+			this.data = data;
+			this.pos = 0;
+			this.mark = 0;
+		}
+
+		@Override
+		public int read() throws IOException
+		{
+			if (pos < data.length)
+			{
+	            return data[pos++] & 0xff;
+			}
+			return -1;
+		}
+
+		public int peek() throws IOException
+		{
+			if (pos < data.length)
+			{
+				return data[pos] & 0xff;
+			}
+			return -1;
+		}
+
+		public int read(byte[] buffer, int offset, int length)
+		{
+			int count = Math.min(data.length, pos + length) - pos;
+			count = (offset + count > buffer.length) ? buffer.length - offset : count;
+
+			System.arraycopy(data, pos, buffer, offset, count);
+			pos += count;
+			return count;
+		}
+
+		public int readShort() throws IOException
+		{
+			return (read() << 8) | read();
+		}
+
+		public int readInt() throws IOException
+		{
+			return (read() << 24) | (read() << 16) | (read() << 8) | read();
+		}
+
+		@Override
+		public int available() throws IOException
+		{
+			return data.length - pos;
+	    }
+
+		public int seek(int position)
+		{
+			if (position >= 0 && position < data.length)
+			{
+				pos = position;
+			}
+			return pos;
+		}
+
+		public int getPosition()
+		{
+			return pos;
+		}
+
+		@Override
+		public void close() throws IOException
+		{
+			data = null;
+			pos = -1;
+		}
+
+		@Override
+		public boolean markSupported()
+		{
+			return true;
+		}
+
+		@Override
+		public synchronized void mark(int readlimit)
+		{
+			mark = pos;
+		}
+
+		@Override
+		public synchronized void reset() throws IOException
+		{
+			pos = mark;
+		}
+
+		@Override
+		public long skip(long n) throws IOException
+		{
+			long skipped = Math.min(n, available());
+			pos += skipped;
+			return skipped;
+	    }
+	}
+
+	private ByteArrayInputStream2 getConfigRulesDecompressed(byte[] bytesToDecompress)
 	{
 		if (bytesToDecompress == null)
 		{
 			return null;
 		}
 
+Trace.beginSection("getCRDecompressed");
+
 		int headerSize = 2 + 4 + 4;
 		int sizeCompressed = bytesToDecompress.length - headerSize;
 		if (sizeCompressed < 0)
@@ -1188,9 +1302,9 @@ public class GameActivity extends $${gameActivitySuperClass}$$ implements Surfac
 
 		// validate signature
 		int signature = buffer.getShort();
-		if (signature != 0x39d8)
+		if (signature != 0x39d8 && signature != 0x39da)
 		{
-			Log.debug("failed sig ");
+			Log.debug("failed sig");
 			return null;
 		}
 		int version = buffer.getInt();
@@ -1198,37 +1312,48 @@ public class GameActivity extends $${gameActivitySuperClass}$$ implements Surfac
 		byte[] bytesDecompressed = new byte[sizeUncompressed];
 
 		// decrypt if key provided
-		try
+		if (ConfigRulesKey != null)
 		{
-			Field keyField = this.getClass().getDeclaredField("CONFIGRULES_KEY");
-			String key = (String)keyField.get(this);
-
-			Cipher cipher = Cipher.getInstance("AES/ECB/PKCS5Padding");
-			cipher.init(Cipher.DECRYPT_MODE, generateAESKey(key));
-			byte[] decrypted = cipher.doFinal(bytesToDecompress, headerSize, sizeCompressed);
-			sizeCompressed = decrypted.length;
+			try
+			{
+Trace.beginSection("decrypt");
+				Cipher cipher = Cipher.getInstance("AES/ECB/PKCS5Padding");
+				cipher.init(Cipher.DECRYPT_MODE, ConfigRulesKey);
+				byte[] decrypted = cipher.doFinal(bytesToDecompress, headerSize, sizeCompressed);
+				sizeCompressed = decrypted.length;
+Trace.endSection();  // decrypt
 
-			// decompress it
-			Inflater inflater = new Inflater();
-			inflater.setInput(decrypted, 0, sizeCompressed);
-			int resultLength = inflater.inflate(bytesDecompressed);
-			inflater.end();
-			if (resultLength != sizeUncompressed)
+				// decompress it
+Trace.beginSection("decompress");
+				Inflater inflater = new Inflater();
+				inflater.setInput(decrypted, 0, sizeCompressed);
+				int resultLength = inflater.inflate(bytesDecompressed);
+				inflater.end();
+Trace.endSection();  // decompress
+				if (resultLength != sizeUncompressed)
+				{
+					Log.debug("1 unexpected size ("+resultLength+", "+sizeUncompressed+")");
+					return null;
+				}
+			}
+			catch (Exception e)
 			{
-				Log.debug("1 unexpected size ("+resultLength+", "+sizeUncompressed+")");
-				return null;
+Trace.endSection();  // paired
+				Log.debug("failed " + e.toString());
+				ConfigRulesKey = null;
 			}
 		}
-		catch (Exception e)
+		if (ConfigRulesKey == null)
 		{
-			Log.debug("failed "+e.toString());
 			// decompress it without decrypting
 			try
 			{
+Trace.beginSection("decompress");
 				Inflater inflater = new Inflater();
 				inflater.setInput(bytesToDecompress, headerSize, sizeCompressed);
 				int resultLength = inflater.inflate(bytesDecompressed);
 				inflater.end();
+Trace.endSection();  // decompress
 				if (resultLength != sizeUncompressed)
 				{
 					Log.debug("2 unexpected size ("+resultLength+", "+sizeUncompressed+")");
@@ -1237,65 +1362,48 @@ public class GameActivity extends $${gameActivitySuperClass}$$ implements Surfac
 			}
 			catch (Exception ee)
 			{
-				Log.debug("failed "+ee.toString());
+Trace.endSection();  // paired
+				Log.debug("failed " + ee.toString());
 				return null;
 			}
 		}
+Trace.endSection();  // getCRDecompressed
 
 		// now provide a BufferedReader for final buffer
-		InputStream inputStream = null;
+		ByteArrayInputStream2 inputStream = null;
 		try
 		{
-			inputStream = new ByteArrayInputStream(bytesDecompressed);
+			inputStream = new ByteArrayInputStream2(bytesDecompressed);
 			return inputStream;
 		}
 		catch (Exception e)
 		{
 			Log.debug("failed "+e.toString());
-        }
+		}
 		return null;
 	}
 
 	// returns InputStream to most up to date version of configrules
-	private InputStream getConfigRulesStream(String ProjectName)
+	private byte[] getBestConfigRulesData(String ProjectName)
 	{
 		ConfigRulesVersion = -1;
 		ConfigRulesCRC = 0;
 		ConfigRulesStatus = "";
 
-		long assetsCRC = 0;
-		long internalCRC = 0;
-		long externalCRC = 0;
-
-		byte[] configBytesAssets = getByteArrayFromAssets("configrules.bin.png");
-		byte[] configBytesInternal = getByteArrayFromFile(InternalFilesDir + "configrules.bin.png");
-		byte[] configBytesExternal = getByteArrayFromFile(android.os.Environment.getExternalStorageDirectory().getAbsolutePath() + "/UnrealGame/" + ProjectName + "/configrules.bin.png");
-
-		if (configBytesAssets != null)
-		{
-			CRC32 crc = new CRC32();
-			crc.update(configBytesAssets);
-			assetsCRC = crc.getValue();
-		}
-
-		if (configBytesInternal != null)
-		{
-			CRC32 crc = new CRC32();
-			crc.update(configBytesInternal);
-			internalCRC = crc.getValue();
-			Log.debug(String.format("CRC32: %08X", internalCRC));
-		}
-
-		if (configBytesExternal != null)
-		{
-			CRC32 crc = new CRC32();
-			crc.update(configBytesExternal);
-			externalCRC = crc.getValue();
-		}
+Trace.beginSection("loadVersions");
+		byte[] configBytesAssets = ConfigRulesBytesAssets != null ? ConfigRulesBytesAssets : getByteArrayFromAssets("configrules.bin.png");
+		byte[] configBytesInternal = ConfigRulesBytesInternal != null ? ConfigRulesBytesInternal : getByteArrayFromFile(InternalFilesDir + "configrules.bin.png");
+		byte[] configBytesExternal = ConfigRulesBytesExternal != null ? ConfigRulesBytesExternal : getByteArrayFromFile(android.os.Environment.getExternalStorageDirectory().getAbsolutePath() + "/UnrealGame/" + ProjectName + "/configrules.bin.png");
+		ConfigRulesBytesAssets = null;
+		ConfigRulesBytesInternal = null;
+		ConfigRulesBytesExternal = null;
+Trace.endSection();  // loadVersions
 
+Trace.beginSection("getVersionCodes");
 		int assetsVersion = getConfigRulesVersion(configBytesAssets);
 		int internalVersion = getConfigRulesVersion(configBytesInternal);
 		int externalVersion = getConfigRulesVersion(configBytesExternal);
+Trace.endSection();
 
 		// is there one at all?
 		if (assetsVersion == -1 && internalVersion == -1 && externalVersion == -1)
@@ -1304,53 +1412,57 @@ public class GameActivity extends $${gameActivitySuperClass}$$ implements Surfac
 			return null;
 		}
 
-		InputStream result = null;
+Trace.beginSection("pickVersion");
+		byte[] result = null;
 		if (CommandLineFlags.containsKey("noconfigrulesdownload"))
 		{
 			ConfigRulesVersion = assetsVersion;
-			ConfigRulesCRC = assetsCRC;
-			result = getConfigRulesDecompressed(configBytesAssets);
+			result = configBytesAssets;
 		}
 		else if (assetsVersion >= internalVersion && assetsVersion >= externalVersion)
 		{
 			ConfigRulesVersion = assetsVersion;
-			ConfigRulesCRC = assetsCRC;
-			result = getConfigRulesDecompressed(configBytesAssets);
+			result = configBytesAssets;
 		}
 		else if (internalVersion >= assetsVersion && internalVersion >= externalVersion)
 		{
 			ConfigRulesVersion = internalVersion;
-			ConfigRulesCRC = internalCRC;
-			result = getConfigRulesDecompressed(configBytesInternal);
+			result = configBytesInternal;
 			if (result == null && assetsVersion != -1)
 			{
 				Log.debug("using assets rule: ");
 				File internalconfigFile = new File(InternalFilesDir + "configrules.bin.png");
 				internalconfigFile.delete();
 				ConfigRulesVersion = assetsVersion;
-				ConfigRulesCRC = assetsCRC;
-				result = getConfigRulesDecompressed(configBytesAssets);
+				result = configBytesAssets;
 			}
 		}
 		else
 		{
 			ConfigRulesVersion = externalVersion;
-			ConfigRulesCRC = externalCRC;
-			result = getConfigRulesDecompressed(configBytesExternal);
+			result = configBytesExternal;
 			if (result == null && internalVersion != -1)
 			{
 				ConfigRulesVersion = internalVersion;
-				ConfigRulesCRC = internalCRC;
-				result = getConfigRulesDecompressed(configBytesInternal);
+				result = configBytesInternal;
 			}
 			if (result == null && assetsVersion != -1)
 			{
 				ConfigRulesVersion = assetsVersion;
-				ConfigRulesCRC = assetsCRC;
-				result = getConfigRulesDecompressed(configBytesAssets);
+				result = configBytesAssets;
 			}
 		}
-		if (result == null)
+Trace.endSection();
+
+		if (result != null)
+		{
+Trace.beginSection("calcCRC");
+			CRC32 crc = new CRC32();
+			crc.update(result);
+			ConfigRulesCRC = crc.getValue();
+Trace.endSection();
+		}
+		else
 		{
 			ConfigRulesVersion = -1;
 			ConfigRulesCRC = 0;
@@ -1617,6 +1729,260 @@ public class GameActivity extends $${gameActivitySuperClass}$$ implements Surfac
 		TextureFormats += "ETC1";
     }
 
+	private void ConfigRules_ApplySets(ByteArrayInputStream2 reader, Map<String, String> variables, StringTable stringTable) throws IOException
+	{
+		int count = (int)reader.read();
+		while (count-- > 0)
+		{
+			int flags = (int)reader.read();
+			String key = stringTable.getString(reader.readShort());
+			String value = stringTable.getString(reader.readShort());
+
+			if ((flags & CRToken.SETVARFLAGS_EXPAND) > 0)
+			{
+				value = ExpandVariables(variables, value);
+			}
+
+			if ((flags & CRToken.SETVARFLAGS_APPEND) > 0)
+			{
+				if (variables.containsKey(key))
+				{
+					value = variables.get(key) + value;
+				}
+				variables.put(key, value);
+			}
+			else
+			{
+				variables.put(key, value);
+			}
+		}
+	}
+
+	private void ConfigRules_ApplyClears(ByteArrayInputStream2 reader, Map<String, String> variables, StringTable stringTable) throws IOException
+	{
+		int count = reader.read();
+		while (count-- > 0)
+		{
+			String key = stringTable.getString(reader.readShort());
+			variables.remove(key);
+		}
+	}
+
+	private boolean ConfigRules_EvalConditions(ByteArrayInputStream2 byteread, Map<String, String> variables, StringTable stringTable) throws IOException
+	{
+		boolean bConditionTrue = true;
+		boolean bSourceIsCommandLine = false;
+		String PreviousRegexMatch = "";
+
+		int groups = byteread.read();
+		while (bConditionTrue && groups-- > 0)
+		{
+			int sourceToken = byteread.read();
+			String Source = stringTable.getString(byteread.readShort());
+			int compareToken = byteread.read();
+			String MatchString = stringTable.getString(byteread.readShort());
+
+			switch (sourceToken)
+			{
+				case CRToken.SOURCE_LITERAL:
+					Source = variables.getOrDefault(Source, Source);
+					break;
+
+				case CRToken.SOURCE_PREVREGEX:
+					Source = PreviousRegexMatch;
+					break;
+
+				case CRToken.SOURCE_COMMANDLINE:
+					bSourceIsCommandLine = true;
+					Source = CommandLineFull;
+					break;
+
+				case CRToken.SOURCE_EXIST:
+					Source = MatchString;
+					break;
+
+				default:
+					SystemInfo_Log("ConfigRules: unknown source token!");
+					bConditionTrue = false;
+					continue;
+			}
+
+			// deal with ignore compares
+			if (compareToken >= CRToken.COMPARE_IGNORE)
+			{
+				compareToken -= CRToken.COMPARE_IGNORE;
+				Source = Source.toLowerCase();
+				// MatchString is already lowercase from parser
+			}
+
+			// continue if passes, break for failure
+			switch (compareToken)
+			{
+				case CRToken.COMPARE_EXIST:
+					if (variables.containsKey(Source))
+					{
+						continue;
+					}
+					break;
+
+				case CRToken.COMPARE_NOTEXIST:
+					if (!variables.containsKey(Source))
+					{
+						continue;
+					}
+					break;
+
+				case CRToken.COMPARE_EQUAL:
+					if (bSourceIsCommandLine)
+					{
+						if (CommandLineFlags.containsKey(MatchString))
+						{
+							continue;
+						}
+					}
+					else if (Source.equals(MatchString))
+					{
+						continue;
+					}
+					break;
+
+				case CRToken.COMPARE_NOTEQUAL:
+					if (bSourceIsCommandLine)
+					{
+						if (!CommandLineFlags.containsKey(MatchString))
+						{
+							continue;
+						}
+					}
+					else if (!Source.equals(MatchString))
+					{
+						continue;
+					}
+					break;
+
+				case CRToken.COMPARE_STARTS:
+					if (Source.startsWith(MatchString))
+					{
+						continue;
+					}
+					break;
+
+				case CRToken.COMPARE_CONTAINS:
+					if (Source.contains(MatchString))
+					{
+						continue;
+					}
+					break;
+
+				case CRToken.COMPARE_ENDS:
+					if (Source.endsWith(MatchString))
+					{
+						continue;
+					}
+					break;
+
+				case CRToken.COMPARE_REGEX:
+					{
+						Pattern pattern = Pattern.compile(MatchString);
+						Matcher matcher = pattern.matcher(Source);
+						if (matcher.find())
+						{
+							int groupCount = matcher.groupCount();
+							PreviousRegexMatch = matcher.group(groupCount >= 1 ? 1 : 0);
+							continue;
+						}
+					}
+					break;
+
+				default:
+					{
+						boolean bNumericOperands = true;
+						float SourceFloat = 0.0f;
+						float MatchFloat = 0.0f;
+
+						// convert source and match to float if numeric
+						try
+						{
+							SourceFloat = Float.parseFloat(Source);
+						}
+						catch (NumberFormatException ne)
+						{
+							bNumericOperands = false;
+						}
+						try
+						{
+							MatchFloat = Float.parseFloat(MatchString);
+						}
+						catch (NumberFormatException ne)
+						{
+							bNumericOperands = false;
+						}
+
+						switch (compareToken)
+						{
+							case CRToken.COMPARE_VEQUAL:
+								{
+									boolean bSourceMatchEqual = Math.abs(SourceFloat - MatchFloat) < 0.00001f;
+									if ((bNumericOperands && bSourceMatchEqual) || (!bNumericOperands && (Source.compareTo(MatchString) == 0)))
+									{
+										continue;
+									}
+								}
+								break;
+
+							case CRToken.COMPARE_VNOTEQUAL:
+								{
+									boolean bSourceMatchEqual = Math.abs(SourceFloat - MatchFloat) < 0.00001f;
+									if ((bNumericOperands && !bSourceMatchEqual) || (!bNumericOperands && (Source.compareTo(MatchString) != 0)))
+									{
+										continue;
+									}
+								}
+								break;
+
+							case CRToken.COMPARE_LESS:
+								if ((bNumericOperands && (SourceFloat < MatchFloat)) || (!bNumericOperands && (Source.compareTo(MatchString) < 0)))
+								{
+									continue;
+								}
+								break;
+
+							case CRToken.COMPARE_LESSEQUAL:
+								if ((bNumericOperands && (SourceFloat <= MatchFloat)) || (!bNumericOperands && (Source.compareTo(MatchString) <= 0)))
+								{
+									continue;
+								}
+								break;
+
+							case CRToken.COMPARE_GREATER:
+								if ((bNumericOperands && (SourceFloat > MatchFloat)) || (!bNumericOperands && (Source.compareTo(MatchString) > 0)))
+								{
+									continue;
+								}
+								break;
+
+							case CRToken.COMPARE_GREATEREQUAL:
+								if ((bNumericOperands && (SourceFloat >= MatchFloat)) || (!bNumericOperands && (Source.compareTo(MatchString) >= 0)))
+								{
+									continue;
+								}
+								break;
+
+							default:
+								SystemInfo_Log("ConfigRules: unknown compare token!");
+								bConditionTrue = false;
+								continue;
+						}
+					}
+					break;
+			}
+
+			// failed comparison
+			bConditionTrue = false;
+		}
+		return bConditionTrue;
+	}
+
 	// removes the entry and exit from ends of the string if present and paired
 	private String RemoveSurrounds(String input, String entry, String exit)
 	{
@@ -1845,6 +2211,54 @@ public class GameActivity extends $${gameActivitySuperClass}$$ implements Surfac
 					break;
 				}
 			}
+			else if (CompareType.equals("CMP_Starts"))
+			{
+				if (!Source.startsWith(MatchString))
+				{
+					bConditionTrue = false;
+					break;
+				}
+			}
+			else if (CompareType.equals("CMP_Contains"))
+			{
+				if (!Source.contains(MatchString))
+				{
+					bConditionTrue = false;
+					break;
+				}
+			}
+			else if (CompareType.equals("CMP_Ends"))
+			{
+				if (!Source.endsWith(MatchString))
+				{
+					bConditionTrue = false;
+					break;
+				}
+			}
+			else if (CompareType.equals("CMP_StartsIgnore"))
+			{
+				if (!Source.toLowerCase().startsWith(MatchString.toLowerCase()))
+				{
+					bConditionTrue = false;
+					break;
+				}
+			}
+			else if (CompareType.equals("CMP_ContainsIgnore"))
+			{
+				if (!Source.toLowerCase().contains(MatchString.toLowerCase()))
+				{
+					bConditionTrue = false;
+					break;
+				}
+			}
+			else if (CompareType.equals("CMP_EndsIgnore"))
+			{
+				if (!Source.toLowerCase().endsWith(MatchString.toLowerCase()))
+				{
+					bConditionTrue = false;
+					break;
+				}
+			}
 			else if (CompareType.equals("CMP_Regex"))
 			{
 				Pattern pattern = Pattern.compile(MatchString);
@@ -1955,6 +2369,53 @@ public class GameActivity extends $${gameActivitySuperClass}$$ implements Surfac
 		return bConditionTrue;
 	}
 
+	private class StringTable
+	{
+		private int count = 0;
+		private int stringLen = 0;
+		private int[] offsets = null;
+		private int[] lengths = null;
+		private byte[] data = null;
+
+		public boolean Init(ByteArrayInputStream2 stream)
+		{
+			try
+			{
+				count = stream.readShort();
+				stringLen = stream.readShort();
+
+				offsets = new int[count];
+				for (int i=0; i<count; i++)
+				{
+					offsets[i] = stream.readShort();
+				}
+
+				lengths = new int[count];
+				for (int i=0; i<count; i++)
+				{
+					lengths[i] = stream.readShort();
+				}
+
+				data = new byte[stringLen];
+				stream.read(data, 0, stringLen);
+			}
+			catch (Exception e)
+			{
+				return false;
+			}
+			return true;
+		}
+
+		public String getString(int index)
+		{
+			if (index < 0 || index >= count)
+			{
+				return null;
+			}
+			return new String(data, offsets[index], lengths[index], StandardCharsets.UTF_8);
+		}
+	}
+
 	protected void processSystemInfo(Lock lock)
 	{
 		SystemInfo_InitLog();
@@ -2146,7 +2607,35 @@ Trace.endSection();
 			lock.unlock();
 		}
 
+Trace.beginSection("ReadCRFiles");
+		ConfigRulesBytesAssets = getByteArrayFromAssets("configrules.bin.png");
+//		ConfigRulesBytesInternal = getByteArrayFromFile(InternalFilesDir + "configrules.bin.png");	// skip here since may be downloaded before runConfigRules
+		ConfigRulesBytesExternal = getByteArrayFromFile(android.os.Environment.getExternalStorageDirectory().getAbsolutePath() + "/UnrealGame/" + ProjectName + "/configrules.bin.png");
+Trace.endSection();
+
+Trace.beginSection("CRKey");
+		try
+		{
+			Field keyField = this.getClass().getDeclaredField("CONFIGRULES_KEY");
+			ConfigRulesKey = generateAESKey((String)keyField.get(this));
+		}
+		catch (Exception e)
+		{
+		}
+Trace.endSection();
+Trace.beginSection("CRCacheKey");
+		try
+		{
+			Field keyField = this.getClass().getDeclaredField("CONFIGRULES_CACHE_KEY");
+			ConfigRulesCacheKey = generateAESKey((String)keyField.get(this));
+		}
+		catch (Exception e)
+		{
+		}
+Trace.endSection();
+
 		// calculate a CRC to use for cache invalidation
+Trace.beginSection("CRVarCRC");
 		CRC32 crc32 = new CRC32();
 		for (Map.Entry<String, String> entry : variables.entrySet())
 		{
@@ -2154,6 +2643,7 @@ Trace.endSection();
 			crc32.update(entry.getValue().getBytes());
 		}
 		ConfigRulesVarsCRC = crc32.getValue();
+Trace.endSection();
 	}
 
 	protected boolean runConfigRules(Map<String, String> variables, String ProjectName, String PackageName)
@@ -2178,19 +2668,10 @@ Trace.endSection();
 		int configRulesVersion = 0;
 		variables.put("configRulesVersion", Integer.toString(configRulesVersion));
 		BufferedReader reader = null;
-		InputStream stream = getConfigRulesStream(ProjectName);
-		if (stream != null)
-		{
-			try
-			{
-				reader = new BufferedReader(new InputStreamReader(stream));
-			}
-			catch (Exception e)
-			{
-				ConfigRulesStatus += "failed to create config rules reader. " + e;
-				reader = null;
-			}
-		}
+Trace.beginSection("getStream");
+		byte[] configRulesData = getBestConfigRulesData(ProjectName);
+		ByteArrayInputStream2 stream = null;
+Trace.endSection();  // getStream
 
 		Log.debug("$$$ configrules: start parsing");
 
@@ -2198,17 +2679,9 @@ Trace.endSection();
 		int configFlags = 0;
 		boolean bWriteCache = true;
 		boolean bClearCache = false;
-		java.security.Key cacheKey = null;
-		try
-		{
-			Field keyField = this.getClass().getDeclaredField("CONFIGRULES_CACHE_KEY");
-			cacheKey = generateAESKey((String)keyField.get(this));
-		}
-		catch (Exception e)
-		{
-		}
-		if (reader != null)
+		if (configRulesData != null)
 		{
+Trace.beginSection("CheckCache");
 			String configCacheFilename = InternalFilesDir + "configrules.cache";
 			byte[] cacheFileContents = getByteArrayFromFile(configCacheFilename);
 			if (cacheFileContents != null)
@@ -2228,19 +2701,22 @@ Trace.endSection();
 					SystemInfo_Log("read: " + String.format("%d, %08X, %08X", cacheVersion, cacheRulesCRC, cacheVarsCRC));
 					if (cacheVersion == ConfigRulesVersion && cacheRulesCRC == ConfigRulesCRC && cacheVarsCRC == ConfigRulesVarsCRC)
 					{
+Trace.beginSection("UseCache");
 						byte[] decrypted = null;
-						if (cacheKey != null)
+						if (ConfigRulesCacheKey != null)
 						{
+Trace.beginSection("CacheDecrypt");
 							try
 							{
 								Cipher cipher = Cipher.getInstance("AES/ECB/PKCS5Padding");
-								cipher.init(Cipher.DECRYPT_MODE, cacheKey);
+								cipher.init(Cipher.DECRYPT_MODE, ConfigRulesCacheKey);
 								decrypted = cipher.doFinal(cacheFileContents, headerSize, sizeEncrypted);
 							}
 							catch (Exception e)
 							{
 								SystemInfo_Log("configrules cache decrypt failure: " + e.toString());
 							}
+Trace.endSection();
 						}
 						else
 						{
@@ -2266,9 +2742,7 @@ Trace.endSection();
 									SystemInfo_LogMap("variables:", variables);
 								}
 
-								reader = null;
-								stream.close();
-								stream = null;
+								configRulesData = null;
 								bWriteCache = false;
 								configFlags = cacheFlags;
 							}
@@ -2278,18 +2752,310 @@ Trace.endSection();
 							}
 						}
 					}
+Trace.endSection();  // UseCache
 				}
 			}
+Trace.endSection();  // CheckCache
 		}
 
-		// process the rules
-		if (reader != null)
+		// need configrules decoded if still valid
+		boolean bPreparsed = false;
+		if (configRulesData != null)
 		{
-			// convert to lowercase once to check against chipset key
-			String hardwareCheck = variables.get("hardware").toLowerCase();
-			
+			// need the configrules data so decode it
+			stream = getConfigRulesDecompressed(configRulesData);
+			try
+			{
+				bPreparsed = stream.peek() < 0x20;
+			}
+			catch (Exception e)
+			{
+				stream = null;
+			}
+		}
+
+		// convert to lowercase once to check against chipset key
+		String hardwareCheck = variables.get("hardware").toLowerCase();
+
+		// deal with preparsed variant of configrules
+		if (bPreparsed)
+		{
+Trace.beginSection("Preparsed");
+			try
+			{
+				ByteArrayInputStream2 byteread = stream;
+
+				int version = byteread.readShort();
+				configRulesVersion = byteread.readShort();
+				variables.put("configRulesVersion", Integer.toString(configRulesVersion));
+				SystemInfo_Log("ConfigRules version: " + configRulesVersion);
+
+				StringTable stringTable = new StringTable();
+				if (!stringTable.Init(byteread))
+				{
+					throw new IOException("Configrules StringTable corruption");
+				}
+
+				String PreviousRegexMatch = "";
+
+				boolean bChangedVars = false;
+				int commandToken;
+				int commandLength = -1;
+
+				int instructionCount = byteread.readShort();
+				int position = byteread.getPosition();
+				int basePosition = position;
+				while (instructionCount-- > 0)
+				{
+					// move to next command
+					byteread.seek(position);
+
+					if (bChangedVars)
+					{
+						bChangedVars = false;
+
+						// see if log message requested
+						if (variables.containsKey("log"))
+						{
+							SystemInfo_Log(variables.get("log"));
+							variables.remove("log");
+						}
+
+						// check if requested to dump cpu info to the log
+						if (variables.containsKey("dumpcpuinfo"))
+						{
+							configFlags |= 2;
+							variables.remove("dumpcpuinfo");
+							if (CpuInfo != null)
+							{
+								SystemInfo_LogMap("cpuInfo:", CpuInfo);
+							}
+						}
+
+						// check if requested to dump variables to the log
+						if (variables.containsKey("dumpvars"))
+						{
+							configFlags |= 1;
+							variables.remove("dumpvars");
+							SystemInfo_LogMap("variables:", variables);
+						}
+
+						// if there was a raised error, stop
+						boolean hasError = variables.containsKey("error");
+						boolean hasWarning = variables.containsKey("warning");
+						if (hasError || hasWarning)
+						{
+							bClearCache = true;
+
+							ForceExitLink = "";
+							if (variables.containsKey("link"))
+							{
+								ForceExitLink = variables.get("link");
+							}
+
+							ForceExitUpdateButtonText = "";
+							if (variables.containsKey("updatebutton"))
+							{
+								ForceExitUpdateButtonText = getResourceStringOrDefault(PackageName, variables.get("updatebutton"), ForceExitUpdateButtonText);
+							}
+
+							ForceExitHelpButtonText = "Help";
+							if (variables.containsKey("helpbutton"))
+							{
+								ForceExitHelpButtonText = getResourceStringOrDefault(PackageName, variables.get("helpbutton"), ForceExitHelpButtonText);
+							}
+
+							ForceExitQuitButtonText = "Quit";
+							if (variables.containsKey("exitbutton"))
+							{
+								ForceExitQuitButtonText = getResourceStringOrDefault(PackageName, variables.get("exitbutton"), ForceExitQuitButtonText);
+							}
+
+							ForceExitContinueButtonText = "Continue";
+							if (variables.containsKey("continuebutton"))
+							{
+								ForceExitContinueButtonText = getResourceStringOrDefault(PackageName, variables.get("continuebutton"), ForceExitContinueButtonText);
+							}
+
+							// localize caption string from string table
+							ForceExitCaption = "Device not supported";
+							if (variables.containsKey("caption"))
+							{
+								ForceExitCaption = getResourceStringOrDefault(PackageName, variables.get("caption"), ForceExitCaption);
+							}
+
+							// check if ignoring the warnings and errors
+							if (bIgnoreDialog)
+							{
+								if (hasWarning)
+								{
+									hasWarning = false;
+									variables.remove("warning");
+								}
+								if (hasError)
+								{
+									hasError = false;
+									variables.remove("error");
+								}
+							}
+
+							if (hasWarning)
+							{
+								SaveWarning = variables.get("warning");
+
+								// localize error string from string table
+								ForceExitMessage = getResourceStringOrDefault(PackageName, SaveWarning, "This device may not run this game.");
+								variables.remove("warning");
+								bForceGameEndWithWarning = true;
+							}
+
+							if (hasError)
+							{
+								SaveError = variables.get("error");
+
+								// localize error string from string table
+								ForceExitMessage = getResourceStringOrDefault(PackageName, SaveError, "This device cannot run this game.");
+								bForceGameEndWithError = true;
+								break;
+							}
+						}
+
+						// stop if user wants to break
+						if (variables.containsKey("break"))
+						{
+							break;
+						}
+					}
+
+					// read command and data length and advance position for next command
+					commandToken = byteread.read();
+					commandLength = byteread.readShort();
+					position += 1 + 2 + commandLength;
+
+					// handle commands
+					switch (commandToken)
+					{
+						case CRToken.COMMAND_EOF:
+							instructionCount = 0;
+							break;
+
+						case CRToken.COMMAND_GOTO:
+							{
+								int offset = byteread.readInt();
+								if (offset > 0)
+								{
+									position = basePosition + offset;
+								}
+							}
+							break;
+
+						case CRToken.COMMAND_SET:
+							ConfigRules_ApplySets(byteread, variables, stringTable);
+							bChangedVars = true;
+							break;
+
+						case CRToken.COMMAND_CLEAR:
+							ConfigRules_ApplyClears(byteread, variables, stringTable);
+							bChangedVars = true;
+							break;
+
+						case CRToken.COMMAND_CHIPSET:
+							// only check if haven't already set chipset
+							if (!variables.containsKey("chipset"))
+							{
+								String key = stringTable.getString(byteread.readShort());
+								if (key.contains(hardwareCheck) || hardwareCheck.contains(key))
+								{
+									variables.put("useAffinity", stringTable.getString(byteread.readShort()));
+									variables.put("chipset", stringTable.getString(byteread.readShort()));
+									variables.put("gpu", stringTable.getString(byteread.readShort()));
+									variables.put("processorCount", stringTable.getString(byteread.readShort()));
+									variables.put("bigCoreMask", stringTable.getString(byteread.readShort()));
+									variables.put("littleCoreMask", stringTable.getString(byteread.readShort()));
+
+									SystemInfo_Log("Chipset: " + variables.get("chipset"));
+									SystemInfo_Log("GPU: " + variables.get("gpu"));
+									bChangedVars = true;
+								}
+							}
+							break;
+
+						case CRToken.COMMAND_CONDITION:
+							{
+								boolean bConditionTrue = ConfigRules_EvalConditions(byteread, variables, stringTable);
+								if (bConditionTrue)
+								{
+									ConfigRules_ApplySets(byteread, variables, stringTable);
+									ConfigRules_ApplyClears(byteread, variables, stringTable);
+									bChangedVars = true;
+								}
+							}
+							break;
+
+						case CRToken.COMMAND_IF:
+							{
+								int offsetFalse = byteread.readInt();
+								boolean bConditionTrue = ConfigRules_EvalConditions(byteread, variables, stringTable);
+								if (!bConditionTrue && offsetFalse > 0)
+								{
+									position = basePosition + offsetFalse;
+								}
+							}
+							break;
+
+						case CRToken.COMMAND_ELSEIF:
+							{
+								int offsetFalse = byteread.readInt();
+								boolean bConditionTrue = ConfigRules_EvalConditions(byteread, variables, stringTable);
+								if (!bConditionTrue && offsetFalse > 0)
+								{
+									position = basePosition + offsetFalse;
+								}
+							}
+							break;
+
+						case CRToken.COMMAND_ELSE:
+							// NOP
+							break;
+
+						case CRToken.COMMAND_ENDIF:
+							// NOP
+							break;
+
+						default:
+							SystemInfo_Log("ConfigRules: unexpected " + commandToken + " encountered!");
+							continue;
+					}
+				}
+			}
+			catch (IOException ie)
+			{
+				SystemInfo_Log("failed to read configuration rules: " + ie);
+				ConfigRulesStatus += "failed to read configuration rules : " + ie+ " ";
+			}
+			if (stream != null)
+			{
+				ConfigRulesStatus += "configuration rules finished. " + ForceExitMessage;
+				try
+				{
+					stream.close();
+				}
+				catch (Exception e)
+				{
+				}
+				stream = null;
+			}
+Trace.endSection();  // Preparsed
+		}
+
+		// process the rules from text form
+		if (stream != null)
+		{
+Trace.beginSection("Parsing");
 			try
 			{
+				reader = new BufferedReader(new InputStreamReader(stream));
+
 				Stack<Integer> StateStack = new Stack<Integer>();
 
 				int currentState = CONFIGRULES_STATE_RUN;
@@ -2719,14 +3485,17 @@ Trace.endSection();
 		// warnings or errors will clear the cache
 		if (bClearCache)
 		{
+Trace.beginSection("ClearCache");
 			File configCacheFile = new File(InternalFilesDir + "configrules.cache");
 			configCacheFile.delete();
 			bWriteCache = false;
+Trace.endSection();  // ClearCache
 		}
 
 		// write cache
 		if (bWriteCache)
 		{
+Trace.beginSection("WriteCache");
 			File configCacheFile = new File(InternalFilesDir + "configrules.cache");
 			try (ByteArrayOutputStream hashStream = new ByteArrayOutputStream())
 			{
@@ -2747,15 +3516,17 @@ Trace.endSection();
 				FileOutputStream fileStream = new FileOutputStream(configCacheFile);
 				fileStream.write(headerBytes);
 
-				if (cacheKey == null)
+				if (ConfigRulesCacheKey == null)
 				{
 					fileStream.write(hashData);
 				}
 				else
 				{
+Trace.beginSection("CacheEncrypt");
 					Cipher cipher = Cipher.getInstance("AES/ECB/PKCS5Padding");
-					cipher.init(Cipher.ENCRYPT_MODE, cacheKey);
+					cipher.init(Cipher.ENCRYPT_MODE, ConfigRulesCacheKey);
 					byte[] encrypted = cipher.doFinal(hashData, 0, hashData.length);
+Trace.endSection();  // CacheEncrypt
 					fileStream.write(encrypted);
 				}
 				fileStream.close();
@@ -2768,6 +3539,7 @@ Trace.endSection();
 			{
 				configCacheFile.delete();
 			}
+Trace.endSection();  // WriteCache
 		}
 
 		// clear Profile if OverrideProfile was not left set as well
@@ -2801,7 +3573,6 @@ Trace.endSection();
 		{
 			SystemInfo_Log("Failed to read memstats: MemstatRSSRequiredDiff");
 		}
-		
 
 		// notify AndroidProcess of the core masks and affinity enable state
 		if (variables.containsKey("unsupportedDevice"))
@@ -2819,6 +3590,21 @@ Trace.endSection();
 	{
 		SystemInfo_DumpLog();
 
+		// check for commandline requests (non-cached!) if not Shipping
+		if (!BuildConfiguration.equals("Shipping") && CommandLineFlags.containsKey("crvars"))
+		{
+			String values = CommandLineFlags.get("crvars");
+			ArrayList<String> pairs = ParseSegments(values, ',', "\"", "\"");
+			for (String pair : pairs)
+			{
+				ArrayList<String> keyvalues = ParseSegments(pair, '=', "\"", "\"");
+				if (keyvalues.size() == 2)
+				{
+					variables.put(keyvalues.get(0), RemoveSurrounds(keyvalues.get(1), "\"", "\""));
+				}
+			}
+		}
+
 		boolean bUseAffinity = variables.get("useAffinity").equals("true");
 		int littleCoreMask = Integer.decode(variables.get("littleCoreMask"));
 		int bigCoreMask = Integer.decode(variables.get("bigCoreMask"));
diff --git a/Engine/Build/Android/Prebuilt/ConfigRulesTool/compile.bat b/Engine/Build/Android/Prebuilt/ConfigRulesTool/compile.bat
--- a/Engine/Build/Android/Prebuilt/ConfigRulesTool/compile.bat
+++ b/Engine/Build/Android/Prebuilt/ConfigRulesTool/compile.bat
@@ -1,6 +1,6 @@
-"%JAVA_HOME%"\bin\javac.exe -d classes src\*.java
+"%JAVA_HOME%"\bin\javac.exe -d classes src\*.java ..\..\Java\src\com\epicgames\unreal\CRToken.java
 cd classes
-"%JAVA_HOME%"\bin\jar cfe ../bin/ConfigRulesTool.jar ConfigRulesTool *.class
+"%JAVA_HOME%"\bin\jar cfe ../bin/ConfigRulesTool.jar ConfigRulesTool *.class com/epicgames/unreal/CRToken.class
 
 ECHO "You MUST re-compile this from a macOS environment since iOS also uses this tool."
 ECHO "We need that since the macOS version of this script will create an app bundle with "
diff --git a/Engine/Build/Android/Prebuilt/ConfigRulesTool/compile.sh b/Engine/Build/Android/Prebuilt/ConfigRulesTool/compile.sh
--- a/Engine/Build/Android/Prebuilt/ConfigRulesTool/compile.sh
+++ b/Engine/Build/Android/Prebuilt/ConfigRulesTool/compile.sh
@@ -1,11 +1,11 @@
 #!/bin/bash
 
 # compiles the .java into .class
-javac -d classes src/*.java
+javac -d classes src/*.java ../../Java/src/com/epicgames/unreal/CRToken.java
 
 pushd classes
 # makes the jar from the class
-jar cfe ../bin/ConfigRulesTool.jar ConfigRulesTool *.class
+jar cfe ../bin/ConfigRulesTool.jar ConfigRulesTool *.class com/epicgames/unreal/CRToken.class
 popd
 
 # creates the .app standalone with the previously generated jar
diff --git a/Engine/Build/Android/Prebuilt/ConfigRulesTool/src/ConfigRulesTool.java b/Engine/Build/Android/Prebuilt/ConfigRulesTool/src/ConfigRulesTool.java
--- a/Engine/Build/Android/Prebuilt/ConfigRulesTool/src/ConfigRulesTool.java
+++ b/Engine/Build/Android/Prebuilt/ConfigRulesTool/src/ConfigRulesTool.java
@@ -8,6 +8,7 @@ import java.nio.ByteBuffer;
 import java.nio.file.Files;
 import java.nio.file.Paths;
 import java.nio.file.Path;
+import java.nio.charset.StandardCharsets;
 
 import javax.crypto.Cipher;
 import javax.crypto.SecretKey;
@@ -15,6 +16,115 @@ import javax.crypto.SecretKeyFactory;
 import javax.crypto.spec.PBEKeySpec;
 import javax.crypto.spec.SecretKeySpec;
 
+import com.epicgames.unreal.CRToken;
+
+	class ByteArrayInputStream2 extends InputStream
+	{
+		private byte[] data;
+		private int pos;
+		private int mark;
+
+		public ByteArrayInputStream2(byte[] data)
+		{
+			this.data = data;
+			this.pos = 0;
+			this.mark = 0;
+		}
+
+		@Override
+		public int read() throws IOException
+		{
+			if (pos < data.length)
+			{
+	            return data[pos++] & 0xff;
+			}
+			return -1;
+		}
+
+		public int peek() throws IOException
+		{
+			if (pos < data.length)
+			{
+				return data[pos] & 0xff;
+			}
+			return -1;
+		}
+
+		public int read(byte[] buffer, int offset, int length)
+		{
+			int count = Math.min(data.length, pos + length) - pos;
+			count = (offset + count > buffer.length) ? buffer.length - offset : count;
+
+			System.arraycopy(data, pos, buffer, offset, count);
+			pos += count;
+			return count;
+		}
+
+		public int readShort() throws IOException
+		{
+			return (read() << 8) | read();
+		}
+
+		public int readInt() throws IOException
+		{
+			return (read() << 24) | (read() << 16) | (read() << 8) | read();
+		}
+
+		@Override
+		public int available() throws IOException
+		{
+			return data.length - pos;
+	    }
+
+		public int seek(int position)
+		{
+			if (position >= 0 && position < data.length)
+			{
+				pos = position;
+			}
+			return pos;
+		}
+
+		public int getPosition()
+		{
+			return pos;
+		}
+
+		@Override
+		public void close() throws IOException
+		{
+			data = null;
+			pos = -1;
+		}
+
+		@Override
+		public boolean markSupported()
+		{
+			return true;
+		}
+
+		@Override
+		public synchronized void mark(int readlimit)
+		{
+			mark = pos;
+		}
+
+		@Override
+		public synchronized void reset() throws IOException
+		{
+			pos = mark;
+		}
+
+		@Override
+		public long skip(long n) throws IOException
+		{
+			long skipped = Math.min(n, available());
+			pos += skipped;
+			return skipped;
+	    }
+	}
+
+
 public class ConfigRulesTool
 {
 	public static String CipherTransform = "AES/ECB/PKCS5Padding";
@@ -51,33 +161,51 @@ public class ConfigRulesTool
 	{
 		String inFilename = "configrules.txt";
 		String outFilename = "configrules.bin.png";
+		String command = "";
 		String key = "";
 		int headerSize = 2 + 4 + 4;
+		boolean bVerbose = false;
 		
-		System.out.println("ConfigRulesTool v1.1\n");
+		System.out.println("ConfigRulesTool v1.2\n");
 
 		if (args.length == 0)
 		{
-			System.out.println("Usage: [op] inputfile outputfile [key]\n");
+			System.out.println("Usage: [op] [-verbose] inputfile outputfile [key]\n");
 			System.out.println("\tc = compress and encrypt if key provided");
 			System.out.println("\td = decrypt and decompress");
+			System.out.println("\tp = parse, compress and encrypt if key provided");
 			System.exit(0);
 		}
 		
-		if (args.length > 1)
+		command = args[0];
+		int nextParam = 1;
+		for (int argIndex=1; argIndex < args.length; argIndex++)
 		{
-			inFilename = args[1];
-		}
-		if (args.length > 2)
-		{
-			outFilename = args[2];
-		}
-		if (args.length > 3)
-		{
-			key = args[3];
-		}
+			if (args[argIndex].equals("-verbose"))
+			{
+				bVerbose = true;
+				continue;
+			}
+			switch (nextParam)
+			{
+				case 1:
+					inFilename = args[argIndex];
+					nextParam++;
+					break;
+					
+				case 2:
+					outFilename = args[argIndex];
+					nextParam++;
+					break;
+
+				case 3:
+					key = args[argIndex];
+					nextParam++;
+					break;
+			}
+		}	
 		
-		if (args[0].equals("c"))
+		if (command.equals("c") || command.equals("p"))
 		{
 			byte[] bytesToCompress = null;
 			Path path = Paths.get(inFilename);
@@ -126,6 +254,17 @@ public class ConfigRulesTool
 				System.exit(-1);
 			}
 			
+			// handle parsing request instead of using raw text
+			if (command.equals("p"))
+			{
+				bytesToCompress = parseConfigRules(bytesToCompress, bVerbose);
+				if (bytesToCompress == null)
+				{
+					System.exit(-1);
+				}
+				sizeUncompressed = bytesToCompress.length;
+			}
+			
 			Deflater deflater = new Deflater();
 			deflater.setInput(bytesToCompress);
 			deflater.finish();
@@ -159,7 +298,7 @@ public class ConfigRulesTool
 				fileOutStream = new FileOutputStream(outFile);
 				byte[] signature = new byte[2];
 				signature[0] = (byte)0x39;
-				signature[1] = (byte)0xd8;
+				signature[1] = command.equals("p") ? (byte)0xda : (byte)0xd8;
 				fileOutStream.write(signature);
 				writeInt(fileOutStream, version);
 				writeInt(fileOutStream, sizeUncompressed);
@@ -204,7 +343,7 @@ public class ConfigRulesTool
 			System.exit(0);
 		}
 	
-		if (args[0].equals("d"))
+		if (command.equals("d"))
 		{
 			byte[] bytesToDecompress = null;
 			Path path = Paths.get(inFilename);
@@ -226,7 +365,8 @@ public class ConfigRulesTool
 			
 			ByteBuffer buffer = ByteBuffer.wrap(bytesToDecompress);
 			int signature = buffer.getShort();
-			if (signature != 0x39d8)
+			boolean bParsed = (signature == 0x39da);
+			if (!bParsed && signature != 0x39d8)
 			{
 				System.out.println("Input file signature is invalid: " + inFilename + ", " + Integer.toString(signature));
 				System.exit(-1);
@@ -273,6 +413,15 @@ public class ConfigRulesTool
 				System.exit(-1);
 			}
 			System.out.println("Version: " + Integer.toString(version) + ", Uncompressed size: " + Integer.toString(sizeUncompressed));
+			
+			if (bParsed)
+			{
+				if ((bytesDecompressed = decompileConfigRules(bytesDecompressed, bVerbose)) == null)
+				{
+					System.out.println("Error decompiling file: " + inFilename);
+					System.exit(-1);
+				}
+			}
 
 			File outFile = new File(outFilename);
 			FileOutputStream fileOutStream = null;
@@ -302,8 +451,1322 @@ public class ConfigRulesTool
 			System.out.println("Wrote file: " + outFilename);
 			System.exit(0);
 		}
-		
-		System.out.println("Unknown op: " + args[0]);
+
+		System.out.println("Unknown op: " + command);
 		System.exit(-1);
 	}
+	
+	// ===== configrules preparser =====
+
+	static int RegExFound = 0;
+	static int RegExRemoved = 0;
+
+	class StringTable
+	{
+		static public HashMap<String, Integer> stringTable = new HashMap<>();
+		static public ByteArrayOutputStream stringOffsets = new ByteArrayOutputStream();
+		static public ByteArrayOutputStream stringLengths = new ByteArrayOutputStream();
+		static public ByteArrayOutputStream stringData = new ByteArrayOutputStream();
+
+		static private int count = 0;
+		static private int stringLen = 0;
+		static private int[] offsets = null;
+		static private int[] lengths = null;
+		static private byte[] data = null;
+		
+		static public int calls = 0;
+
+		static int count()
+		{
+			return stringTable.size();
+		}
+		
+		static void close() throws IOException
+		{
+			//stringOffsets.close();
+			//stringLengths.close();
+			//stringData.close();
+		}
+
+		static int addString(String input)
+		{
+			calls++;
+			int nextIndex = count();
+			int findIndex = stringTable.getOrDefault(input, nextIndex);
+			if (findIndex == nextIndex)
+			{
+				stringTable.put(input, nextIndex);
+				
+				int offset = stringData.size();
+				stringOffsets.write(offset >> 8);
+				stringOffsets.write(offset & 255);
+
+				int len = input.length();
+				stringLengths.write(len >> 8);
+				stringLengths.write(len & 255);
+
+				stringData.write(input.getBytes(), 0, len);
+			}
+			return findIndex;
+		}
+
+		static boolean Init(ByteArrayInputStream2 stream)
+		{
+			try
+			{
+				count = stream.readShort();
+				stringLen = stream.readShort();
+
+				offsets = new int[count];
+				for (int i=0; i<count; i++)
+				{
+					offsets[i] = stream.readShort();
+				}
+
+				lengths = new int[count];
+				for (int i=0; i<count; i++)
+				{
+					lengths[i] = stream.readShort();
+				}
+
+				data = new byte[stringLen];
+				stream.read(data, 0, stringLen);
+			}
+			catch (Exception e)
+			{
+				return false;
+			}
+			return true;
+		}
+
+		static String getString(int index)
+		{
+			if (index < 0 || index >= count)
+			{
+				return null;
+			}
+			return new String(data, offsets[index], lengths[index], StandardCharsets.UTF_8);
+		}
+	}
+	
+	// removes the entry and exit from ends of the string if present and paired
+	private static String RemoveSurrounds(String input, String entry, String exit)
+	{
+		int entryLength = entry.length();
+		int exitLength = exit.length();
+		int inputLength = input.length();
+
+		if (inputLength >=2 && entryLength > 0 && entryLength == exitLength)
+		{
+			char start = input.charAt(0);
+			char end = input.charAt(inputLength-1);
+
+			for (int index=0; index < entryLength; index++)
+			{
+				if (entry.charAt(index) == start && exit.charAt(index) == end)
+				{
+					return input.substring(1, inputLength-1);
+				}
+			}
+		}
+		return input;
+	}
+
+	// returns list of strings separated by split using entry/exit as pairing sets (ex. "(" and ")").  The entry and exit characters need to be in same order
+	private static ArrayList<String> ParseSegments(String input, char split, String entry, String exit)
+	{
+		ArrayList<String> output = new ArrayList<String>();
+		Stack<Integer> entryStack = new Stack<Integer>();
+
+		int startIndex = 0;
+		int scanIndex = 0;
+		int exitIndex = -1;
+		int inputLength = input.length();
+		while (scanIndex < inputLength)
+		{
+			char scan = input.charAt(scanIndex);
+
+			if (scan == split && entryStack.empty())
+			{
+				output.add(input.substring(startIndex, scanIndex).trim());
+				scanIndex++;
+				startIndex = scanIndex;
+				continue;
+			}
+			scanIndex++;
+
+			if (scan == '\\')
+			{
+				scanIndex++;
+				continue;
+			}
+
+			if (!entryStack.empty() && exit.indexOf(scan) == exitIndex)
+			{
+				entryStack.pop();
+				exitIndex = entryStack.empty() ? -1 : entryStack.peek();
+				continue;
+			}
+
+			int entryIndex = entry.indexOf(scan);
+			if (entryIndex >= 0)
+			{
+				entryStack.add(entryIndex);
+				exitIndex = entryIndex;
+				continue;
+			}
+		}
+		if (startIndex < inputLength)
+		{
+			output.add(input.substring(startIndex).trim());
+		}
+
+		return output;
+	}
+
+	static String TestRegEx(String input)
+	{
+		// special RegEx symbols that could be escaped
+		String regex = ".*+?|$^()[]{}-\\";
+		int regexLength = regex.length();
+					
+		// remove each escaped symbol then check if it is still in string
+		String test = input;
+		for (int rIndex = 0; rIndex < regexLength; rIndex++)
+		{
+			String symbol = "" + regex.charAt(rIndex);
+			String replace = "\\" + symbol;
+			int escapeIndex = test.indexOf(replace);
+			while (escapeIndex >= 0)
+			{
+				test = test.substring(0, escapeIndex) + (escapeIndex - 2 < test.length() ? test.substring(escapeIndex + 2) : "");
+				escapeIndex = test.indexOf(replace);
+			}
+			if (test.indexOf(symbol) >= 0)
+			{
+				// can't replace it
+				return null;
+			}
+		}
+					
+		// make a version without any escapes
+		for (int rIndex = 0; rIndex < regexLength; rIndex++)
+		{
+			String symbol = "" + regex.charAt(rIndex);
+			String replace = "\\" + symbol;
+			int escapeIndex = input.indexOf(replace);
+			while (escapeIndex >= 0)
+			{
+				input = input.substring(0, escapeIndex) + (escapeIndex - 1 < test.length() ? input.substring(escapeIndex + 1) : "");
+				escapeIndex = input.indexOf(replace);
+			}
+		}
+					
+		return input;
+	}
+	
+	private static boolean parseConditions(ByteArrayOutputStream stream, ArrayList<String> conditions, int lineNumber, boolean bVerbose)
+	{
+		int condLen = conditions.size();
+		stream.write(condLen);
+
+		boolean bFirst = false;//true;
+		for (String condition : conditions)
+		{
+			int sourceToken = -1;
+			int compareToken = -1;
+			String SourceType = "";
+			String CompareType = "";
+			String MatchString = "";
+			int foundBits = 0;
+
+			// deal with condition group (src,cmp,match)
+			ArrayList<String> groups = ParseSegments(RemoveSurrounds(condition, "(", ")"), ',', "\"", "\"");
+			for (String group : groups)
+			{
+				ArrayList<String> keyvalue = ParseSegments(group, '=', "\"", "\"");
+				if (keyvalue.size() == 2)
+				{
+					String key = RemoveSurrounds(keyvalue.get(0), "\"", "\"");
+					String value = RemoveSurrounds(keyvalue.get(1), "\"", "\"");
+
+					int keyToken = CRToken.findCondKey(key);
+					switch (keyToken)
+					{
+						case CRToken.CONDKEY_SOURCETYPE:
+							sourceToken = CRToken.findSource(value);
+							if (sourceToken == -1)
+							{
+								sourceToken = CRToken.SOURCE_LITERAL;
+							}
+							SourceType = value;
+							foundBits |= 1;
+							break;
+							
+						case CRToken.CONDKEY_COMPARETYPE:
+							compareToken = CRToken.findCompare(value);
+							CompareType = value;
+							if (compareToken == CRToken.COMPARE_UNKNOWN)
+							{
+								System.out.println("Error: condition comparetype '" + value + "' invalid, line " + lineNumber);
+								return false;
+							}
+							foundBits |= 2;
+							break;
+							
+						case CRToken.CONDKEY_MATCHSTRING:
+							MatchString = value;
+							foundBits |= 4;
+							break;
+							
+						default:
+							System.out.println("Error: condition type '" + key + "' invalid, line " + lineNumber);
+							return false;
+							
+					}
+				}
+			}
+			
+			// make sure we got all 3
+			if (foundBits != 0x07)
+			{
+				if ((foundBits & 1) == 0)
+				{
+					System.out.println("Error: SourceType missing, line " + lineNumber);
+				}
+				if ((foundBits & 2) == 0)
+				{
+					System.out.println("Error: CompareType missing, line " + lineNumber);
+				}
+				if ((foundBits & 4) == 0)
+				{
+					System.out.println("Error: MatchString missing, line " + lineNumber);
+				}
+				return false;
+			}
+			
+			// go ahead and make lowercase if ignore compare type
+			if (compareToken >= CRToken.COMPARE_IGNORE)
+			{
+				MatchString = MatchString.toLowerCase();
+			}
+
+			// attempt to optimize Regex (it is expensive)
+			if (compareToken == CRToken.COMPARE_REGEX)
+			{
+				RegExFound++;
+
+				if (MatchString.startsWith("^"))
+				{
+					// possibly COMPARETYPE_STARTS
+					String test = TestRegEx(MatchString.substring(1));
+					if (test != null)
+					{
+						compareToken = CRToken.COMPARE_STARTS;
+						CompareType = "CMP_Starts";
+						MatchString = test;
+						RegExRemoved++;
+						if (bVerbose)
+						{
+							System.out.println("--RegEx optimize (" + SourceType + "," + CompareType + ",\"" + MatchString + "\"), line " + lineNumber);
+						}
+					}
+				}
+				else if (MatchString.endsWith("$") && !MatchString.endsWith("\\$"))
+				{
+					// possibly COMPARETYPE_ENDS
+					String test = TestRegEx(MatchString.substring(0, MatchString.length() - 1));
+					if (test != null)
+					{
+						compareToken = CRToken.COMPARE_ENDS;
+						CompareType = "CMP_Ends";
+						MatchString = test;
+						RegExRemoved++;
+						if (bVerbose)
+						{
+							System.out.println("--RegEx optimize (" + SourceType + "," + CompareType + ",\"" + MatchString + "\"), line " + lineNumber);
+						}
+					}
+				}
+				else
+				{
+					// possibly COMPARE_TYPE_CONTAINS
+					String test = TestRegEx(MatchString);
+					if (test != null)
+					{
+						compareToken = CRToken.COMPARE_CONTAINS;
+						CompareType = "CMP_Contains";
+						MatchString = test;
+						RegExRemoved++;
+						if (bVerbose)
+						{
+							System.out.println("--RegEx optimize (" + SourceType + "," + CompareType + ",\"" + MatchString + "\"), line " + lineNumber);
+						}
+					}
+				}
+			}
+
+			if (bFirst)
+			{
+				bFirst = false;
+				System.out.println("--(" + SourceType + "," + CompareType + "," + MatchString + ")");
+			}
+
+			// write condition
+			stream.write(sourceToken);
+			int sourceIndex = StringTable.addString(SourceType);
+			stream.write(sourceIndex >> 8);
+			stream.write(sourceIndex & 255);
+			stream.write(compareToken);
+			int matchIndex = StringTable.addString(MatchString);
+			stream.write(matchIndex >> 8);
+			stream.write(matchIndex & 255);
+		}
+		
+		return true;
+	}
+
+	private static byte[] parseConfigRules(byte[] sourceData, boolean bVerbose)
+	{
+		int configRulesVersion = 0;
+
+		int lineNumber = 0;
+		int instructionCount = 0;
+
+		CRToken.Init();
+		
+		BufferedReader reader = null;
+		InputStream stream = new ByteArrayInputStream(sourceData);
+		if (stream != null)
+		{
+			try
+			{
+				reader = new BufferedReader(new InputStreamReader(stream));
+			}
+			catch (Exception e)
+			{
+				System.out.println("Failed to create config rules reader. " + e);
+				return null;
+			}
+		}
+		else
+		{
+			System.out.println("Unable to create ByteArrayInputStream.");
+			return null;
+		}
+		try
+		{
+			ByteArrayOutputStream memStream = new ByteArrayOutputStream();
+			
+			// doesn't really need to be a hash but not critical
+			HashMap<Integer, Integer> gotoFixups = new HashMap<Integer, Integer>();
+			Stack<Integer> endifOffsets = new Stack<Integer>();
+			Stack<Stack<Integer>> recoverOffsets = new Stack<Stack<Integer>>();
+			Stack<Integer> falseOffsets = new Stack<Integer>();
+			int falseOffset = -1;
+			int ifDepth = 0;
+
+			String line;
+			while ((line = reader.readLine()) != null)
+			{
+				lineNumber++;
+				line = line.trim();
+				if (line.length() < 1)
+				{
+					continue;
+				}
+				if (line.startsWith("//") || line.startsWith(";"))
+				{
+					if (line.startsWith("// version:"))
+					{
+						configRulesVersion = Integer.parseInt(line.substring(11));
+						System.out.println("ConfigRules version: " + configRulesVersion);
+					}
+					continue;
+				}
+
+				// look for command
+				int index = line.indexOf(':');
+				if (index == -1)
+				{
+					continue;
+				}
+				String original = line;
+				String command = line.substring(0, index).trim();
+				line = line.substring(index + 1).trim();
+
+				int commandToken = CRToken.findCommand(command);
+				switch (commandToken)
+				{
+					case CRToken.COMMAND_SET:
+						// set:(a=b[,c=d,...])
+						ArrayList<String> sets = ParseSegments(RemoveSurrounds(line, "(", ")"), ',', "(\"", ")\"");
+						int setsLen = sets.size();
+						if (setsLen == 0)
+						{
+							System.out.println("Error: set has no sets, line " + lineNumber);
+							return null;
+						}
+						else if (setsLen > 255)
+						{
+							System.out.println("Error: set has too many sets, line " + lineNumber);
+							return null;
+						}
+
+						try
+						{
+							ByteArrayOutputStream setStream = new ByteArrayOutputStream();
+							setStream.write(setsLen);
+							for (String assignment : sets)
+							{
+								ArrayList<String> keyvalue = ParseSegments(assignment, '=', "\"", "\"");
+								if (keyvalue.size() == 2)
+								{
+									String key = RemoveSurrounds(keyvalue.get(0), "\"", "\"");
+									String value = RemoveSurrounds(keyvalue.get(1), "\"", "\"");
+									
+									int flags = CRToken.SETVARFLAGS_NONE;
+									if (key.startsWith("APPEND_"))
+									{
+										flags |= CRToken.SETVARFLAGS_APPEND;
+										key = key.substring(7);
+									}
+									if (value.contains("$("))
+									{
+										flags |= CRToken.SETVARFLAGS_EXPAND;
+									}
+
+									setStream.write(flags);
+
+									int keyIndex = StringTable.addString(key);
+									setStream.write(keyIndex >> 8);
+									setStream.write(keyIndex & 255);
+
+									int valueIndex = StringTable.addString(value);
+									setStream.write(valueIndex >> 8);
+									setStream.write(valueIndex & 255);
+								}
+								else
+								{
+									System.out.println("Error: set has mismatched key/value pairs, line " + lineNumber);
+									return null;
+								}
+							}
+							//setStream.close();
+							byte[] setData = setStream.toByteArray();
+							int setDataLen = setData.length;
+
+							instructionCount++;
+							memStream.write(commandToken);
+							memStream.write(setDataLen >> 8);
+							memStream.write(setDataLen & 255);
+							memStream.write(setData, 0, setDataLen);
+						}
+						catch (Exception e)
+						{
+							System.out.println("Error: set exception, line " + lineNumber + ": " + e);
+						}
+						break;
+
+					case CRToken.COMMAND_CLEAR:
+						// clear:(a[,b,...])
+						ArrayList<String> clears = ParseSegments(RemoveSurrounds(line, "(", ")"), ',', "(\"", ")\"");
+						int clearsLen = clears.size();
+						if (clearsLen == 0)
+						{
+							System.out.println("Error: clear has no sets, line " + lineNumber);
+							return null;
+						}
+						else if (clearsLen > 255)
+						{
+							System.out.println("Error: clear has too many variables, line " + lineNumber);
+							return null;
+						}
+
+						try
+						{
+							ByteArrayOutputStream clearStream = new ByteArrayOutputStream();
+							clearStream.write(clearsLen);
+
+							for (String entry : clears)
+							{	
+								String key = RemoveSurrounds(entry, "\"", "\"");
+
+								int keyIndex = StringTable.addString(key);
+								clearStream.write(keyIndex >> 8);
+								clearStream.write(keyIndex & 255);
+							}
+							//clearStream.close();
+							byte[] clearData = clearStream.toByteArray();
+							int clearDataLen = clearData.length;
+
+							instructionCount++;
+							memStream.write(commandToken);
+							memStream.write(clearDataLen >> 8);
+							memStream.write(clearDataLen & 255);
+							memStream.write(clearData, 0, clearDataLen);
+						}
+						catch (Exception e)
+						{
+							System.out.println("Error: clear exception, line " + lineNumber + ": " + e);
+						}
+						break;
+
+					case CRToken.COMMAND_CHIPSET:
+						// chipset:"hardware",useAffinity,"chipset","cpu",processorCount,bigCoreMask,littleCoreMask
+						ArrayList<String> values = ParseSegments(line, ',', "\"", "\"");
+						if (values.size() == 7)
+						{
+							try
+							{
+								ByteArrayOutputStream chipStream = new ByteArrayOutputStream();
+
+								String originalKey = RemoveSurrounds(values.get(0), "\"", "\"");
+								String key = originalKey.toLowerCase();
+								int keyIndex = StringTable.addString(key);
+								chipStream.write(keyIndex >> 8);
+								chipStream.write(keyIndex & 255);
+
+								for (int i=1; i<7; i++)
+								{
+									String value = RemoveSurrounds(values.get(i), "\"", "\"");
+									int valueIndex = StringTable.addString(value);
+									chipStream.write(valueIndex >> 8);
+									chipStream.write(valueIndex & 255);
+								}
+								
+								// add a string with delta offsets for uppercase characters for decompile (not needed for runtime)
+								{
+									String deltas = "";
+									int lastScan = 0;
+									for (int scan=0; scan < originalKey.length(); scan++)
+									{
+										if (originalKey.charAt(scan) != key.charAt(scan))
+										{
+											int delta = scan - lastScan;
+											lastScan = scan;
+											deltas += Character.toString(delta);
+										}
+									}
+									int deltaIndex = StringTable.addString(deltas);
+									chipStream.write(deltaIndex >> 8);
+									chipStream.write(deltaIndex & 255);
+								}
+								
+								//chipStream.close();
+								byte[] chipData = chipStream.toByteArray();
+								int chipDataLen = chipData.length;
+
+								instructionCount++;
+								memStream.write(commandToken);
+								memStream.write(chipDataLen >> 8);
+								memStream.write(chipDataLen & 255);
+								memStream.write(chipData, 0, chipDataLen);
+							}
+							catch (Exception e)
+							{
+								System.out.println("Error: chipset exception, line " + lineNumber + ": " + e);
+							}
+						}
+						else
+						{
+							System.out.println("Error: chipset has missing values, line " + lineNumber);
+							return null;
+						}
+						break;
+
+					case CRToken.COMMAND_CONDITION:
+						// condition:((SourceType=SRC_DeviceMake,CompareType=CMP_Equal,MatchString="samsung")),(SourceType=,CompareType=,MatchString=),...]),(a=b[,c=d,...]),(a[,b,...])
+						// if all the conditions are true, execute the optional sets and/or clears
+						ArrayList<String> conditionAndSets = ParseSegments(line, ',', "(\"", ")\"");
+						int setsize = conditionAndSets.size();
+						if (setsize == 2 || setsize == 3)
+						{
+							ArrayList<String> conditions = ParseSegments(RemoveSurrounds(conditionAndSets.get(0), "(", ")"), ',', "(\"", ")\"");
+							ArrayList<String> csets = ParseSegments(RemoveSurrounds(conditionAndSets.get(1), "(", ")"), ',', "(\"", ")\"");
+							ArrayList<String> cclears = (setsize == 3) ? ParseSegments(RemoveSurrounds(conditionAndSets.get(2), "(", ")"), ',', "(\"", ")\"") : new ArrayList<String>();
+
+							int condLen = conditions.size();
+							if (condLen == 0)
+							{
+								System.out.println("Error: condition has no conditions, line " + lineNumber);
+								return null;
+							}
+							else if (condLen > 255)
+							{
+								System.out.println("Error: condition has too many conditions, line " + lineNumber);
+							}
+
+							try
+							{
+								ByteArrayOutputStream condStream = new ByteArrayOutputStream();
+
+								// deal with conditions parsing
+								if (!parseConditions(condStream, conditions, lineNumber, bVerbose))
+								{
+									return null;
+								}
+
+								// sets
+								int csetsLen = csets.size();
+								if (csetsLen > 255)
+								{
+									System.out.println("Error: condition has too many sets, line " + lineNumber);
+									return null;
+								}
+
+								try
+								{
+									ByteArrayOutputStream setStream = new ByteArrayOutputStream();
+									setStream.write(csetsLen);
+									for (String assignment : csets)
+									{
+										ArrayList<String> keyvalue = ParseSegments(assignment, '=', "\"", "\"");
+										if (keyvalue.size() == 2)
+										{
+											String key = RemoveSurrounds(keyvalue.get(0), "\"", "\"");
+											String value = RemoveSurrounds(keyvalue.get(1), "\"", "\"");
+
+											int flags = CRToken.SETVARFLAGS_NONE;
+											if (key.startsWith("APPEND_"))
+											{
+												flags |= CRToken.SETVARFLAGS_APPEND;
+												key = key.substring(7);
+											}
+											if (value.contains("$("))
+											{
+												flags |= CRToken.SETVARFLAGS_EXPAND;
+											}
+
+											setStream.write(flags);
+
+											int keyIndex = StringTable.addString(key);
+											setStream.write(keyIndex >> 8);
+											setStream.write(keyIndex & 255);
+
+											int valueIndex = StringTable.addString(value);
+											setStream.write(valueIndex >> 8);
+											setStream.write(valueIndex & 255);
+										}
+										else
+										{
+											System.out.println("Error: condition set has mismatched key/value pairs line " + lineNumber);
+											return null;
+										}
+									}
+									//setStream.close();
+									byte[] setData = setStream.toByteArray();
+									int setDataLen = setData.length;
+									condStream.write(setData, 0, setDataLen);
+								}
+								catch (Exception e)
+								{
+									System.out.println("Error: condition exception, line " + lineNumber + ": " + e);
+								}
+
+								// clears
+								int cclearsLen = cclears.size();
+								if (cclearsLen > 255)
+								{
+									System.out.println("Error: condition has too many clears, line " + lineNumber);
+									return null;
+								}
+							
+								try
+								{
+									ByteArrayOutputStream clearStream = new ByteArrayOutputStream();
+									clearStream.write(cclearsLen);
+
+									for (String entry : cclears)
+									{	
+										String key = RemoveSurrounds(entry, "\"", "\"");
+
+										int keyIndex = StringTable.addString(key);
+										clearStream.write(keyIndex >> 8);
+										clearStream.write(keyIndex & 255);
+									}
+									//clearStream.close();
+									byte[] clearData = clearStream.toByteArray();
+									int clearDataLen = clearData.length;
+									condStream.write(clearData, 0, clearDataLen);
+								}
+								catch (Exception e)
+								{
+									System.out.println("Error: clear exception, line " + lineNumber + ": " + e);
+								}
+								
+								//condStream.close();
+								byte[] condData = condStream.toByteArray();
+								int condDataLen = condData.length;
+								
+								// write it all
+								instructionCount++;
+								memStream.write(commandToken);
+								memStream.write(condDataLen >> 8);
+								memStream.write(condDataLen & 255);
+								memStream.write(condData, 0, condDataLen);
+							}
+							catch (Exception e)
+							{
+								System.out.println("Error: condition exception, line " + lineNumber + ": " + e);
+							}
+						}
+						else
+						{
+							if (setsize == 0)
+							{
+								System.out.println("Error: condition missing statement, line " + lineNumber);
+							}
+							else
+							{
+								System.out.println("Error: condition contains extra segments, line " + lineNumber);
+							}
+							return null;
+						}
+						break;
+
+					case CRToken.COMMAND_IF:
+						{
+							// if:((SourceType=SRC_DeviceMake,CompareType=CMP_Equal,MatchString="samsung")),(SourceType=,CompareType=,MatchString=),...])
+							ArrayList<String> conditions = ParseSegments(RemoveSurrounds(line, "(", ")"), ',', "(\"", ")\"");
+
+							int condLen = conditions.size();
+							if (condLen == 0)
+							{
+								System.out.println("Error: if has no conditions, line " + lineNumber);
+								return null;
+							}
+							else if (condLen > 255)
+							{
+								System.out.println("Error: if has too many conditions, line " + lineNumber);
+							}
+
+							try
+							{
+								ByteArrayOutputStream condStream = new ByteArrayOutputStream();
+
+								// placeholder for later patching (false case jump)
+								condStream.write(0);
+								condStream.write(0);
+								condStream.write(0);
+								condStream.write(0);
+
+								// deal with conditions parsing
+								if (!parseConditions(condStream, conditions, lineNumber, bVerbose))
+								{
+									return null;
+								}
+
+								//condStream.close();
+								byte[] condData = condStream.toByteArray();
+								int condDataLen = condData.length;
+
+								// write it all
+								instructionCount++;
+								memStream.write(commandToken);
+								memStream.write(condDataLen >> 8);
+								memStream.write(condDataLen & 255);
+
+								// remember location for false fixup and save previous endifs
+								ifDepth++;
+								recoverOffsets.push(endifOffsets);
+								endifOffsets = new Stack<Integer>();
+								falseOffsets.push(falseOffset);
+								falseOffset = memStream.size();
+	
+								memStream.write(condData, 0, condDataLen);
+							}
+							catch (Exception e)
+							{
+								System.out.println("Error: if exception, line " + lineNumber + ": " + e);
+							}
+						}
+						break;
+
+					case CRToken.COMMAND_ELSEIF:
+						if (ifDepth > 0)
+						{
+							if (falseOffset == -1)
+							{
+								System.out.println("Error: elseif after else , line " + lineNumber);
+								return null;
+							}
+
+							// write goto
+							instructionCount++;
+							memStream.write(CRToken.COMMAND_GOTO);
+							memStream.write(0);
+							memStream.write(4);
+							// remember offset to fix up for endif
+							endifOffsets.push(memStream.size());
+							memStream.write(0);
+							memStream.write(0);
+							memStream.write(0);
+							memStream.write(0);
+
+							// elseif:((SourceType=SRC_DeviceMake,CompareType=CMP_Equal,MatchString="samsung")),(SourceType=,CompareType=,MatchString=),...])
+							ArrayList<String> conditions = ParseSegments(RemoveSurrounds(line, "(", ")"), ',', "(\"", ")\"");
+
+							int condLen = conditions.size();
+							if (condLen == 0)
+							{
+								System.out.println("Error: elseif has no conditions, line " + lineNumber);
+								return null;
+							}
+							else if (condLen > 255)
+							{
+								System.out.println("Error: elseif has too many conditions, line " + lineNumber);
+							}
+
+							try
+							{
+								ByteArrayOutputStream condStream = new ByteArrayOutputStream();
+
+								// placeholder for later patching (false case jump)
+								condStream.write(0);
+								condStream.write(0);
+								condStream.write(0);
+								condStream.write(0);
+
+								// deal with conditions parsing
+								if (!parseConditions(condStream, conditions, lineNumber, bVerbose))
+								{
+									return null;
+								}
+
+								//condStream.close();
+								byte[] condData = condStream.toByteArray();
+								int condDataLen = condData.length;
+
+								// write it all
+								int here = memStream.size();
+								instructionCount++;
+								memStream.write(commandToken);
+								memStream.write(condDataLen >> 8);
+								memStream.write(condDataLen & 255);
+
+								// save falseOffset goto fixup to here
+								gotoFixups.put(falseOffset, here);
+								falseOffset = memStream.size();
+
+								memStream.write(condData, 0, condDataLen);
+							}
+							catch (Exception e)
+							{
+								System.out.println("Error: if exception, line " + lineNumber + ": " + e);
+							}
+						}
+						else
+						{
+							System.out.println("Error: elseif without an if, line " + lineNumber);
+							return null;
+						}
+						break;
+
+					case CRToken.COMMAND_ELSE:
+						if (ifDepth > 0)
+						{
+							if (falseOffset == -1)
+							{
+								System.out.println("Error: double else , line " + lineNumber);
+								return null;
+							}
+							
+							// write goto
+							instructionCount++;
+							memStream.write(CRToken.COMMAND_GOTO);
+							memStream.write(0);
+							memStream.write(4);
+							// remember offset to fix up for endif
+							endifOffsets.push(memStream.size());
+							memStream.write(0);
+							memStream.write(0);
+							memStream.write(0);
+							memStream.write(0);
+
+							int here = memStream.size();
+
+							// save falseOffset goto fixup to here
+							if (falseOffset != -1)
+							{
+								gotoFixups.put(falseOffset, here);
+							}
+							// clear it since we're at the end of the chain
+							falseOffset = -1;
+
+							// this is a nop.. don't have to write it but makes decompile easier
+							instructionCount++;
+							memStream.write(commandToken);
+							memStream.write(0);
+							memStream.write(0);
+						}
+						else
+						{
+							System.out.println("Error: else without an if, line " + lineNumber);
+							return null;
+						}
+						break;
+
+					case CRToken.COMMAND_ENDIF:
+						if (ifDepth > 0)
+						{
+							// position for false and true gotos
+							int here = memStream.size();
+							
+							// save falseOffset goto for here
+							if (falseOffset != -1)
+							{
+								gotoFixups.put(falseOffset, here);
+							}
+							// recover from any previous if
+							falseOffset = falseOffsets.pop();
+
+							// add goto fixups for endif
+							while (!endifOffsets.empty())
+							{
+								gotoFixups.put(endifOffsets.pop(), here);
+							}
+							// recover from any previous if
+							endifOffsets = recoverOffsets.pop();
+							
+							// this is a nop.. don't have to write it but makes decompile easier
+							instructionCount++;
+							memStream.write(commandToken);
+							memStream.write(0);
+							memStream.write(0);
+							
+							ifDepth--;
+						}
+						else
+						{
+							System.out.println("Error: endif without an if, line " + lineNumber);
+							return null;
+						}
+						break;
+
+					default:
+						System.out.println("Error: unknown command '" + command +"' line " + lineNumber);
+						return null;
+				}
+			}
+
+			// makes sure last variable update is checked
+			instructionCount++;
+			memStream.write(CRToken.COMMAND_EOF);
+			memStream.write(0);
+			memStream.write(0);
+
+			reader.close();
+			
+			//memStream.close();
+			int memStreamLen = memStream.size();
+
+			StringTable.close();
+			int stringOffsetsLen = StringTable.stringOffsets.size();
+			int stringLengthsLen = StringTable.stringLengths.size();
+			int stringHeaderLen = stringOffsetsLen + stringLengthsLen;
+			int stringDataLen = StringTable.stringData.size();
+			int stringLen = stringHeaderLen + stringDataLen;
+			int stringCount = StringTable.count();
+
+			ByteArrayOutputStream dataStream = new ByteArrayOutputStream(2 + 2 + 2 + 2 + stringLen + 2 + memStreamLen);
+			dataStream.write(0);	// indicate binary (anything before 0x20)
+			dataStream.write(1);	// version
+			dataStream.write(configRulesVersion >> 8);
+			dataStream.write(configRulesVersion & 255);
+			dataStream.write(stringCount >> 8);
+			dataStream.write(stringCount & 255);
+			dataStream.write(stringDataLen >> 8);
+			dataStream.write(stringDataLen & 255);
+			dataStream.write(StringTable.stringOffsets.toByteArray(), 0, stringOffsetsLen);
+			dataStream.write(StringTable.stringLengths.toByteArray(), 0, stringLengthsLen);
+			dataStream.write(StringTable.stringData.toByteArray(), 0, stringDataLen);
+			dataStream.write(instructionCount >> 8);
+			dataStream.write(instructionCount & 255);
+
+			byte[] memStreamBytes = memStream.toByteArray();
+			for (Map.Entry<Integer, Integer> entry : gotoFixups.entrySet())
+			{
+				int dest = entry.getKey();
+				int offset = entry.getValue();
+				
+				memStreamBytes[dest++] = (byte)((offset >> 24) & 0xFF);
+				memStreamBytes[dest++] = (byte)((offset >> 16) & 0xFF); 
+				memStreamBytes[dest++] = (byte)((offset >> 8) & 0xFF); 
+				memStreamBytes[dest] = (byte)(offset & 0xFF); 
+			}
+			dataStream.write(memStreamBytes, 0, memStreamLen);
+			
+			System.out.println("Instruction count: " + (instructionCount-1) + " from " + lineNumber + " lines");
+			System.out.println("String count: " + stringCount + " from " + StringTable.calls);
+			System.out.println("Regex removed: " + RegExRemoved + " of " + RegExFound);
+			return dataStream.toByteArray();
+		}
+		catch (IOException ie)
+		{
+			System.out.println("failed to read configuration rules: " + ie);
+			return null;
+		}
+		// unreachable
+	}
+	
+	// =====
+
+
+	public static void writeString(ByteArrayOutputStream stream, String output)
+	{
+		stream.write(output.getBytes(), 0, output.length());
+		System.out.print(output);
+	}
+
+	public static void dumpConditions(ByteArrayInputStream2 byteread, ByteArrayOutputStream dataStream) throws IOException
+	{
+		writeString(dataStream,"(");
+		int groups = byteread.read();
+		while (groups-- > 0)
+		{
+			int sourceToken = byteread.read();
+			String Source = StringTable.getString(byteread.readShort());
+			int compareToken = byteread.read();
+			String MatchString = StringTable.getString(byteread.readShort());
+
+			if (sourceToken == CRToken.SOURCE_EXIST)
+			{
+				Source = "\"[EXIST]\"";
+			}
+
+			writeString(dataStream,"(SourceType=" + Source +",CompareType=" + CRToken.getCompare(compareToken) + ",MatchString=\"" + MatchString + "\")" + (groups > 0 ? "," : ""));
+		}
+		writeString(dataStream,")");
+	}
+
+	private static byte[] decompileConfigRules(byte[] sourceData, boolean bDebug)
+	{
+		try
+		{
+			ByteArrayOutputStream dataStream = new ByteArrayOutputStream();
+
+			CRToken.Init();
+			CRToken.InitStrings();
+
+			ByteArrayInputStream2 byteread = new ByteArrayInputStream2(sourceData);
+			
+			int version = byteread.readShort();
+			int configRulesVersion = byteread.readShort();
+			writeString(dataStream, "// version:" + configRulesVersion + "\n");
+
+			StringTable.Init(byteread);
+			
+			int commandToken;
+			int commandLength = -1;
+
+			int instructionCount = byteread.readShort();
+			int position = byteread.getPosition();
+			int basePosition = position;
+
+			String curPos = "";
+			int ifDepth = 0;
+
+			final char indentChar = '\t';
+			final int indentMult = 1;
+
+			while (instructionCount-- > 0)
+			{
+				while (instructionCount-- > 0)
+				{
+					// move to next command
+					byteread.seek(position);
+
+					// read command and data length and advance position for next command
+					int currentPosition = position - basePosition;
+					commandToken = byteread.read();
+					commandLength = byteread.readShort();
+					position += 1 + 2 + commandLength;
+					
+					curPos = (bDebug ? "" + String.format("%5d:  ", currentPosition) : "") + new String(new char[ifDepth*indentMult]).replace('\0', indentChar);
+
+					switch (commandToken)
+					{
+						case CRToken.COMMAND_EOF:
+							instructionCount = 0;
+							break;
+							
+						case CRToken.COMMAND_GOTO:
+							{
+								int offset = byteread.readInt();
+								if (bDebug)
+								{
+									writeString(dataStream, curPos + ";goto: " + offset + "\n");
+								}
+							}
+							break;
+							
+						case CRToken.COMMAND_SET:
+							{
+								writeString(dataStream, curPos + "set:(");
+								int sets = byteread.read();
+								while (sets-- > 0)
+								{
+									int flags = (int)byteread.read();
+									String key = StringTable.getString(byteread.readShort());
+									String value = StringTable.getString(byteread.readShort());
+									if ((flags & CRToken.SETVARFLAGS_APPEND) > 0)
+									{
+										key = "APPEND_" + key;
+									}
+									writeString(dataStream, key + "=\"" + value + (sets > 0 ? "\"," : "\""));
+								}
+								writeString(dataStream, ")\n");
+							}
+							break;
+							
+						case CRToken.COMMAND_CLEAR:
+							{
+								writeString(dataStream, curPos + "clear:(");
+								int clears = byteread.read();
+								while (clears-- > 0)
+								{
+									String key = StringTable.getString(byteread.readShort());
+									writeString(dataStream, key + (clears > 0 ? "," : ""));
+								}
+								writeString(dataStream, ")\n");
+							}
+							break;
+
+						case CRToken.COMMAND_CHIPSET:
+							{
+								writeString(dataStream, curPos + "chipset:");
+								String key = StringTable.getString(byteread.readShort());
+								String useAffinity = StringTable.getString(byteread.readShort());
+								String chipset = StringTable.getString(byteread.readShort());;
+								String gpu = StringTable.getString(byteread.readShort());
+								String processorCount = StringTable.getString(byteread.readShort());
+								String bigCoreMask = StringTable.getString(byteread.readShort());
+								String littleCoreMask = StringTable.getString(byteread.readShort());
+								String deltas = StringTable.getString(byteread.readShort());
+
+								if (deltas.length() > 0)
+								{
+									char[] keyArray = key.toCharArray();
+									int currentIndex = 0;
+									for (int index=0; index < deltas.length(); index++)
+									{
+										currentIndex += deltas.charAt(index);
+										keyArray[currentIndex] = Character.toUpperCase(keyArray[currentIndex]);
+									}
+									key = String.valueOf(keyArray);
+								}
+
+								writeString(dataStream, "\"" + key + "\", ");
+								writeString(dataStream, useAffinity + ", ");
+								writeString(dataStream, "\"" + chipset + "\", ");
+								writeString(dataStream, "\"" + gpu + "\", ");
+								writeString(dataStream, processorCount + ", ");
+								writeString(dataStream, bigCoreMask + ", ");
+								writeString(dataStream, littleCoreMask + "\n");
+							}
+							break;
+
+						case CRToken.COMMAND_CONDITION:
+							{
+								writeString(dataStream, curPos + "condition:");
+
+								dumpConditions(byteread, dataStream);
+
+								int sets = byteread.read();
+								if (sets > 0)
+								{
+									writeString(dataStream, ", (");
+									while (sets-- > 0)
+									{
+										int flags = (int)byteread.read();
+										String key = StringTable.getString(byteread.readShort());
+										String value = StringTable.getString(byteread.readShort());
+										if ((flags & CRToken.SETVARFLAGS_APPEND) > 0)
+										{
+											key = "APPEND_" + key;
+										}
+										writeString(dataStream, key + "=\"" + value + (sets > 0 ? "\"," : "\""));
+									}
+									writeString(dataStream, ")");
+								}
+								else
+								{
+									writeString(dataStream, ",");
+								}
+								int clears = byteread.read();
+								if (clears > 0)
+								{
+									writeString(dataStream, ", (");
+									while (clears-- > 0)
+									{
+										String key = StringTable.getString(byteread.readShort());
+										writeString(dataStream, key + (clears > 0 ? "," : ""));
+									}
+									writeString(dataStream, ")");
+								}
+								writeString(dataStream, "\n");
+							}
+							break;
+
+						case CRToken.COMMAND_IF:
+							{
+								int offsetFalse = byteread.readInt();
+								ifDepth++;
+
+								writeString(dataStream, curPos + "if:");
+								dumpConditions(byteread, dataStream);
+								if (bDebug)
+								{
+									writeString(dataStream, "\n;   false: " + offsetFalse);
+								}
+								writeString(dataStream, "\n");
+							}
+							break;
+
+						case CRToken.COMMAND_ELSEIF:
+							{
+								int offsetFalse = byteread.readInt();
+
+								writeString(dataStream, curPos.substring(0, curPos.length() - indentMult) + "elseif:");
+								dumpConditions(byteread, dataStream);
+								if (bDebug)
+								{
+									writeString(dataStream, "\n;   false: " + offsetFalse);
+								}
+								writeString(dataStream, "\n");
+							}
+							break;
+
+						case CRToken.COMMAND_ELSE:
+							{
+								writeString(dataStream, curPos.substring(0, curPos.length() - indentMult) + "else:");
+								writeString(dataStream, "\n");
+							}
+							break;
+
+						case CRToken.COMMAND_ENDIF:
+							{
+								writeString(dataStream, curPos.substring(0, curPos.length() - indentMult) + "endif:");
+								writeString(dataStream, "\n");
+								ifDepth--;
+							}
+							break;
+							
+						default:
+							break;
+					}
+				}
+			}
+			
+			return dataStream.toByteArray();
+		}
+		catch (IOException ie)
+		{
+			System.out.println("failed to decompile configuration rules: " + ie);
+			return null;
+		}
+		// unreachable
+	}
+
 }
\ No newline at end of file
diff --git a/Engine/Source/Programs/UnrealBuildTool/Platform/Android/UEDeployAndroid.cs b/Engine/Source/Programs/UnrealBuildTool/Platform/Android/UEDeployAndroid.cs
--- a/Engine/Source/Programs/UnrealBuildTool/Platform/Android/UEDeployAndroid.cs
+++ b/Engine/Source/Programs/UnrealBuildTool/Platform/Android/UEDeployAndroid.cs
@@ -3909,7 +3909,7 @@ namespace UnrealBuildTool
 
 		private void CreateGradlePropertiesFiles(AndroidToolChain ToolChain, UnrealArch Arch, int MinSDKVersion, int TargetSDKVersion, string CompileSDKVersion, string BuildToolsVersion, string PackageName,
 			string? DestApkName, string NDKArch, string UnrealBuildFilesPath, string GameBuildFilesPath, string UnrealBuildGradleAppPath, string UnrealBuildPath, string UnrealBuildGradlePath,
-			bool bForDistribution, bool bIsEmbedded, List<string> OBBFiles)
+			bool bForDistribution, bool bIsEmbedded, List<string> OBBFiles, bool bJavaPerfettoMarkers)
 		{
 			// Create gradle.properties
 			StringBuilder GradleProperties = new StringBuilder();
@@ -4151,8 +4151,6 @@ namespace UnrealBuildTool
 			GradleBuildAdditionsContent.AppendLine("}");
 
 			String BuildGradleAdditionsDefault = "";
-			bool bJavaPerfettoMarkers = false;
-			Ini.GetBool("/Script/AndroidRuntimeSettings.AndroidRuntimeSettings", "bJavaPerfettoMarkers", out bJavaPerfettoMarkers);
 			if (bJavaPerfettoMarkers)
 			{
 				BuildGradleAdditionsDefault = @"
@@ -4231,6 +4229,19 @@ dependencies.implementation('androidx.tracing:tracing:1.2.0')
 			return GetDontBundleLibrariesInAPK(ProjectFile, ForceDontBundleLibrariesInAPK, Configuration, bIsArchive, bFromMSBuild, bIsFromUAT, bVerbose ? Logger : null);
 		}
 
+		private bool GetJavaPerfettoMarkers(ConfigHierarchy Ini, UnrealTargetConfiguration Configuration, bool bDistribution)
+		{
+			// Ship for distribution should not include markers
+			if (bDistribution && Configuration == UnrealTargetConfiguration.Shipping)
+			{
+				return false;
+			}
+
+			bool bJavaPerfettoMarkers = false;
+			Ini.GetBool("/Script/AndroidRuntimeSettings.AndroidRuntimeSettings", "bJavaPerfettoMarkers", out bJavaPerfettoMarkers);
+			return bJavaPerfettoMarkers;
+		}
+
 		// Architecture remapping
 		private static readonly Dictionary<string, string> ArchRemapping = new()
 		{
@@ -4314,6 +4325,8 @@ dependencies.implementation('androidx.tracing:tracing:1.2.0')
 
 			bool bIsBuildMachine = Unreal.IsBuildMachine();
 
+			bool bJavaPerfettoMarkers = GetJavaPerfettoMarkers(Ini, Configuration, bForDistribution);
+
 			// do this here so we'll stop early if there is a problem with the SDK API level (cached so later calls will return the same)
 			string SDKAPILevel = GetSdkApiLevel(ToolChain);
 			string BuildToolsVersion = GetBuildToolsVersion();
@@ -5071,8 +5084,8 @@ dependencies.implementation('androidx.tracing:tracing:1.2.0')
 				WriteJavaOBBDataFile(UnrealOBBDataFileName, PackageName, RequiredOBBFiles, CookFlavor, bPackageDataInsideApk, Arch);
 
 				// update GameActivity.java and GameApplication.java if out of date
-				UpdateGameActivity(Arch, NDKArch, EngineDirectory, UnrealBuildPath, bDontBundleLibrariesInAPK);
-				UpdateGameApplication(Arch, NDKArch, EngineDirectory, UnrealBuildPath);
+				UpdateGameActivity(Arch, NDKArch, EngineDirectory, UnrealBuildPath, bDontBundleLibrariesInAPK, bJavaPerfettoMarkers);
+				UpdateGameApplication(Arch, NDKArch, EngineDirectory, UnrealBuildPath, bJavaPerfettoMarkers);
 
 				// we don't actually need the SO for the bSkipGradleBuild case
 				string? SourceSOName = null;
@@ -5291,7 +5304,7 @@ popd
 				File.WriteAllText(LocalPropertiesFilename, LocalProperties.ToString());
 
 				CreateGradlePropertiesFiles(ToolChain, Arch, MinSDKVersion, TargetSDKVersion, CompileSDKVersion, BuildToolsVersion, PackageName, DestApkName, NDKArch,
-					UnrealBuildFilesPath, GameBuildFilesPath, UnrealBuildGradleAppPath, UnrealBuildPath, UnrealBuildGradlePath, bForDistribution, bSkipGradleBuild, RequiredOBBFiles);
+					UnrealBuildFilesPath, GameBuildFilesPath, UnrealBuildGradleAppPath, UnrealBuildPath, UnrealBuildGradlePath, bForDistribution, bSkipGradleBuild, RequiredOBBFiles, bJavaPerfettoMarkers);
 
 				if (!bSkipGradleBuild)
 				{
@@ -6085,7 +6098,7 @@ popd
 			return "empty";
 		}
 
-		private void UpdateGameActivity(UnrealArch UnrealArch, string NDKArch, string EngineDir, string UnrealBuildPath, bool bDontBundleLibrariesInAPK)
+		private void UpdateGameActivity(UnrealArch UnrealArch, string NDKArch, string EngineDir, string UnrealBuildPath, bool bDontBundleLibrariesInAPK, bool bJavaPerfettoMarkers)
 		{
 			string SourceFilename = Path.Combine(EngineDir, "Build", "Android", "Java", "src", "com", "epicgames", "unreal", "GameActivity.java.template");
 			string DestFilename = Path.Combine(UnrealBuildPath, "src", "com", "epicgames", "unreal", "GameActivity.java");
@@ -6168,8 +6181,6 @@ import java.util.Collection;
 					break;
 			}
 
-			bool bJavaPerfettoMarkers = false;
-			Ini.GetBool("/Script/AndroidRuntimeSettings.AndroidRuntimeSettings", "bJavaPerfettoMarkers", out bJavaPerfettoMarkers);
 			if (bJavaPerfettoMarkers)
 			{
 				GameActivityImportAdditionsDefault += "import androidx.tracing.Trace;\n";
@@ -6303,13 +6314,19 @@ import java.util.Collection;
 			}
 		}
 
-		private void UpdateGameApplication(UnrealArch UnrealArch, string NDKArch, string EngineDir, string UnrealBuildPath)
+		private void UpdateGameApplication(UnrealArch UnrealArch, string NDKArch, string EngineDir, string UnrealBuildPath, bool bJavaPerfettoMarkers)
 		{
 			string SourceFilename = Path.Combine(EngineDir, "Build", "Android", "Java", "src", "com", "epicgames", "unreal", "GameApplication.java.template");
 			string DestFilename = Path.Combine(UnrealBuildPath, "src", "com", "epicgames", "unreal", "GameApplication.java");
 
+			String GameApplicationImportAdditionsDefault = "";
+			if (bJavaPerfettoMarkers)
+			{
+				GameApplicationImportAdditionsDefault += "import androidx.tracing.Trace;\n";
+			}
+
 			Dictionary<string, string> Replacements = new Dictionary<string, string>{
-				{ "//$${gameApplicationImportAdditions}$$", UPL!.ProcessPluginNode(NDKArch, "gameApplicationImportAdditions", "")},
+				{ "//$${gameApplicationImportAdditions}$$", UPL!.ProcessPluginNode(NDKArch, "gameApplicationImportAdditions", GameApplicationImportAdditionsDefault)},
 				{ "//$${gameApplicationOnCreateAdditions}$$", UPL.ProcessPluginNode(NDKArch, "gameApplicationOnCreateAdditions", "")},
 				{ "//$${gameApplicationAttachBaseContextAdditions}$$", UPL.ProcessPluginNode(NDKArch, "gameApplicationAttachBaseContextAdditions", "")},
 				{ "//$${gameApplicationOnLowMemoryAdditions}$$", UPL.ProcessPluginNode(NDKArch, "gameApplicationOnLowMemoryAdditions", "")},
@@ -6326,18 +6343,61 @@ import java.util.Collection;
 			{
 				string SrcLine = TemplateSrc[LineIndex];
 				bool Changed = false;
-				foreach (KeyValuePair<string, string> KVP in Replacements)
+
+				if (!bJavaPerfettoMarkers)
 				{
-					if (SrcLine.Contains(KVP.Key))
+					if (SrcLine.StartsWith("Trace."))
 					{
-						SrcLine = SrcLine.Replace(KVP.Key, KVP.Value);
+						TemplateSrc[LineIndex] = "";
 						Changed = true;
+						TemplateChanged = true;
 					}
 				}
-				if (Changed)
+
+				if (!Changed)
 				{
-					TemplateSrc[LineIndex] = SrcLine;
-					TemplateChanged = true;
+					foreach (KeyValuePair<string, string> KVP in Replacements)
+					{
+						if (SrcLine.Contains(KVP.Key))
+						{
+							if (bJavaPerfettoMarkers)
+							{
+								SrcLine = SrcLine.Replace(KVP.Key, KVP.Value);
+							}
+							else
+							{
+								int TraceStartIndex = KVP.Value.IndexOf("Trace.");
+								if (TraceStartIndex == 0 || (TraceStartIndex > 0 && KVP.Value[TraceStartIndex - 1] == '\n'))
+								{
+									String Addition = KVP.Value;
+									do
+									{
+										int TraceEndIndex = Addition.IndexOf("\n", TraceStartIndex);
+										if (TraceEndIndex < 0)
+										{
+											Addition = Addition.Substring(0, TraceStartIndex);
+										}
+										else
+										{
+											Addition = Addition.Substring(0, TraceStartIndex) + Addition.Substring(TraceEndIndex + 1);
+										}
+										TraceStartIndex = Addition.IndexOf("Trace.", TraceStartIndex);
+									} while (TraceStartIndex >= 0);
+									SrcLine = SrcLine.Replace(KVP.Key, Addition);
+								}
+								else
+								{
+									SrcLine = SrcLine.Replace(KVP.Key, KVP.Value);
+								}
+							}
+							Changed = true;
+						}
+					}
+					if (Changed)
+					{
+						TemplateSrc[LineIndex] = SrcLine;
+						TemplateChanged = true;
+					}
 				}
 			}
 
