Improve -crvars (no need for quotes, - option to clear variable, logging changes)

#android

diff --git a/Engine/Build/Android/Java/src/com/epicgames/unreal/GameActivity.java.template b/Engine/Build/Android/Java/src/com/epicgames/unreal/GameActivity.java.template
--- a/Engine/Build/Android/Java/src/com/epicgames/unreal/GameActivity.java.template
+++ b/Engine/Build/Android/Java/src/com/epicgames/unreal/GameActivity.java.template
@@ -1075,10 +1075,10 @@ public class GameActivity extends $${gameActivitySuperClass}$$ implements Surfac
 			if (argument.startsWith("-"))
 			{
 				argument = argument.substring(1);
-				ArrayList<String> keyvalues = ParseSegments(argument, '=', "\"", "\"");
-				if (keyvalues.size() == 2)
+				int equalIndex = argument.indexOf('=');
+				if (equalIndex > 0)
 				{
-					CommandLineFlags.put(keyvalues.get(0).toLowerCase(), RemoveSurrounds(keyvalues.get(1), "\"", "\""));
+					CommandLineFlags.put(argument.substring(0, equalIndex).toLowerCase(), RemoveSurrounds(argument.substring(equalIndex+1), "\"", "\""));
 				}
 				else
 				{
@@ -3596,10 +3596,25 @@ Trace.endSection();  // WriteCache
 			for (String pair : pairs)
 			{
 				ArrayList<String> keyvalues = ParseSegments(pair, '=', "\"", "\"");
-				if (keyvalues.size() == 2)
+				if (keyvalues.size() > 0)
 				{
-					variables.put(keyvalues.get(0), RemoveSurrounds(keyvalues.get(1), "\"", "\""));
+					String key = keyvalues.get(0);
+					if (keyvalues.size() == 2)
+					{
+						String value = RemoveSurrounds(keyvalues.get(1), "\"", "\"");
+						Log.debug("crvars override: set " + key + " = " + value);
+						variables.put(key, value);
+						continue;
+					}
+					else if (keyvalues.size() == 1 && key.startsWith("-"))
+					{
+						key = key.substring(1);
+						Log.debug("crvars override: remove " + key);
+						variables.remove(key);
+						continue;
+					}
 				}
+				Log.debug("crvars ignoring malformed parameter: " + pair);
 			}
 		}
 
