Replaced reading bundle value with a generated build constant IS_SHIPPING_CONFIG. Replaced some nativeIsShippingBuild() usages with IS_SHIPPING_CONFIG to better support mixed build/config scenarios

#android
#tests tested that both shipping and development builds build as expected

diff --git a/Engine/Build/Android/Java/src/com/epicgames/unreal/GameActivity.java.template b/Engine/Build/Android/Java/src/com/epicgames/unreal/GameActivity.java.template
--- a/Engine/Build/Android/Java/src/com/epicgames/unreal/GameActivity.java.template
+++ b/Engine/Build/Android/Java/src/com/epicgames/unreal/GameActivity.java.template
@@ -259,6 +259,8 @@ public class GameActivity extends $${gameActivitySuperClass}$$ implements Surfac
 //$${gameActivityImplementsAdditions}$$
 															ComponentCallbacks2
 {
+	public static final boolean IS_SHIPPING_CONFIG = $${gameActivityIsShippingConfig}$$;
+	
 	public int VersionCode = 0;
 	public String VersionName = "Unknown";
 	public static int targetSdkVersion = 0;
@@ -754,7 +756,7 @@ public class GameActivity extends $${gameActivitySuperClass}$$ implements Surfac
 			gameControllerManager.onStart();
 		}
 
-		if (!BuildConfiguration.equals("Shipping"))
+		if (!IS_SHIPPING_CONFIG)
 		{
 			// Create console command broadcast listener
 			Log.debug( "Creating console command broadcast listener");
@@ -995,7 +997,7 @@ public class GameActivity extends $${gameActivitySuperClass}$$ implements Surfac
 		if (bUseExternalFilesDir)
 		{
 			BaseDirectory = getExternalFilesDir(null).getAbsolutePath();
-			if (BuildConfiguration.equals("Shipping"))
+			if (IS_SHIPPING_CONFIG)
 			{
 				BaseDirectory = getFilesDir().getAbsolutePath();
 			}
@@ -2656,7 +2658,7 @@ Trace.endSection();
 		String SaveWarning = "";
 
 		// commandline flag may be used to skip dialog if not a shipping build
-		boolean bIgnoreDialog = !BuildConfiguration.equals("Shipping")  // !nativeIsShippingBuild()
+		boolean bIgnoreDialog = !IS_SHIPPING_CONFIG  // !nativeIsShippingBuild()
 								&& CommandLineFlags.containsKey("noconfigrulesdialog");
 
 		// NOTE: the "earlyrestart" variable and anything added to variables by <gameActivityPreConfigRulesParseAdditions> not included in ConfigRulesVarsCRC to invalidate cache!
@@ -3589,7 +3591,7 @@ Trace.endSection();  // WriteCache
 		SystemInfo_DumpLog();
 
 		// check for commandline requests (non-cached!) if not Shipping
-		if (!BuildConfiguration.equals("Shipping") && CommandLineFlags.containsKey("crvars"))
+		if (!IS_SHIPPING_CONFIG && CommandLineFlags.containsKey("crvars"))
 		{
 			String values = CommandLineFlags.get("crvars");
 			ArrayList<String> pairs = ParseSegments(values, ',', "\"", "\"");
@@ -3785,7 +3787,7 @@ Trace.beginSection("loadLibraries");
 Trace.endSection();
 
 		// Bug with certain devices calling onCreate twice, even with singleTask set and System.exit called in onDestroy.
-		if (nativeIsShippingBuild())
+		if (IS_SHIPPING_CONFIG)
 		{
 			if (_activity != null)
 			{
@@ -3920,7 +3922,7 @@ Trace.endSection();
 			UseDisplayCutout = _extrasBundle.getString("UseDisplayCutout") != null;
 
 			String CmdLine = _extrasBundle.getString("cmdline");
-			if (CmdLine != null && !BuildConfiguration.equals("Shipping"))
+			if (CmdLine != null && !IS_SHIPPING_CONFIG)
 			{
 				CommandLineFull = CmdLine;
 			}
@@ -4008,7 +4010,7 @@ Trace.endSection();
 			}
 		}
 
-		boolean bIsShipping = nativeIsShippingBuild();
+		boolean bIsShipping = IS_SHIPPING_CONFIG;
 
 		// create splashscreen dialog (if launched by SplashActivity)
 		if (!isStandalone && _extrasBundle != null)
@@ -8750,21 +8752,9 @@ Trace.endSection();
 		// We need to decide on enabling memory tracing before loading native code.
 		try
 		{
-			if (_bundle != null)
-			{
-				if (!"Shipping".equals(BuildConfiguration))
-				{ 
-					EnableMemoryTracing(AppendDebugCommandLine(GetNotShippingCommandLine(_bundle, CommandLine)));
-				}
-			}
-			else
-			{
-				Bundle Bundle = GameApplication.getAppContext().getPackageManager().getApplicationInfo(getPackageName(), PackageManager.GET_META_DATA).metaData;
-
-				if (!"Shipping".equals(Bundle.getString("com.epicgames.unreal.GameActivity.BuildConfiguration")))
-				{ 
-					EnableMemoryTracing(AppendDebugCommandLine(GetNotShippingCommandLine(Bundle, CommandLine)));
-				}
+			if (!IS_SHIPPING_CONFIG)
+			{ 
+				EnableMemoryTracing(AppendDebugCommandLine(GetNotShippingCommandLine(_bundle != null ? _bundle : GameApplication.getAppContext().getPackageManager().getApplicationInfo(getPackageName(), PackageManager.GET_META_DATA).metaData, CommandLine)));
 			}
 		}
 		catch (ErrnoException | NameNotFoundException e)
diff --git a/Engine/Build/Android/Java/src/com/epicgames/unreal/GameApplication.java.template b/Engine/Build/Android/Java/src/com/epicgames/unreal/GameApplication.java.template
--- a/Engine/Build/Android/Java/src/com/epicgames/unreal/GameApplication.java.template
+++ b/Engine/Build/Android/Java/src/com/epicgames/unreal/GameApplication.java.template
@@ -11,7 +11,10 @@ import androidx.lifecycle.ProcessLifecycleOwner;
 import com.epicgames.unreal.network.NetworkChangedManager;
 //$${gameApplicationImportAdditions}$$
 
-public class GameApplication extends Application implements LifecycleObserver {
+public class GameApplication extends Application implements LifecycleObserver
+{
+	public static final boolean IS_SHIPPING_CONFIG = $${gameApplicationIsShippingConfig}$$;
+	
 	private static final Logger Log = new Logger("UE", "GameApp");
 
 	private static boolean isForeground = false;
diff --git a/Engine/Build/Android/Java/src/com/epicgames/unreal/SplashActivity.java b/Engine/Build/Android/Java/src/com/epicgames/unreal/SplashActivity.java
--- a/Engine/Build/Android/Java/src/com/epicgames/unreal/SplashActivity.java
+++ b/Engine/Build/Android/Java/src/com/epicgames/unreal/SplashActivity.java
@@ -63,7 +63,6 @@ public class SplashActivity extends Activity
 
 		boolean ShouldHideUI = false;
 		boolean UseDisplayCutout = false;
-		boolean IsShipping = false;
 		boolean UseExternalFilesDir = false;
 		try {
 			ApplicationInfo ai = pm.getApplicationInfo(packageName, PackageManager.GET_META_DATA);
@@ -77,11 +76,6 @@ public class SplashActivity extends Activity
 			{
 				UseDisplayCutout = bundle.getBoolean("com.epicgames.unreal.GameActivity.bUseDisplayCutout");
 			}
-			if (bundle.containsKey("com.epicgames.unreal.GameActivity.BuildConfiguration"))
-			{
-				String Configuration = bundle.getString("com.epicgames.unreal.GameActivity.BuildConfiguration");
-				IsShipping = Configuration.equals("Shipping");
-			}
 			if (bundle.containsKey("com.epicgames.unreal.GameActivity.bUseExternalFilesDir"))
             {
                 UseExternalFilesDir = bundle.getBoolean("com.epicgames.unreal.GameActivity.bUseExternalFilesDir");
diff --git a/Engine/Build/Android/Java/src/com/epicgames/unreal/WebViewControl.java b/Engine/Build/Android/Java/src/com/epicgames/unreal/WebViewControl.java
--- a/Engine/Build/Android/Java/src/com/epicgames/unreal/WebViewControl.java
+++ b/Engine/Build/Android/Java/src/com/epicgames/unreal/WebViewControl.java
@@ -118,7 +118,7 @@ class WebViewControl
 			@Override
 			public void run()
 			{
-				WebView.setWebContentsDebuggingEnabled(bEnableRemoteDebugging && !GameActivity.Get().nativeIsShippingBuild());
+				WebView.setWebContentsDebuggingEnabled(bEnableRemoteDebugging && !GameActivity.IS_SHIPPING_CONFIG);
 				
 				// create the WebView
 				webView = new GLWebView(GameActivity.Get());
diff --git a/Engine/Plugins/Runtime/AndroidFileServer/Source/AndroidFileServer/AndroidFileServer_UPL.xml b/Engine/Plugins/Runtime/AndroidFileServer/Source/AndroidFileServer/AndroidFileServer_UPL.xml
--- a/Engine/Plugins/Runtime/AndroidFileServer/Source/AndroidFileServer/AndroidFileServer_UPL.xml
+++ b/Engine/Plugins/Runtime/AndroidFileServer/Source/AndroidFileServer/AndroidFileServer_UPL.xml
@@ -500,7 +500,7 @@ import com.epicgames.unreal.RemoteFileManager;
 		String EngineVersion = "5.0.0";
 		String ProjectName = packageName;
 		ProjectName = ProjectName.substring(ProjectName.lastIndexOf('.') + 1);
-		boolean IsShipping = false;
+		boolean IsShipping = IS_SHIPPING_CONFIG;
 		boolean UseExternalFilesDir = false;
 		boolean PublicLogFiles = false;
 		boolean AllowNetworkConnection = false;
@@ -516,11 +516,6 @@ import com.epicgames.unreal.RemoteFileManager;
 			{
 				ProjectName = bundle.getString("com.epicgames.unreal.GameActivity.ProjectName");
 			}
-			if (bundle.containsKey("com.epicgames.unreal.GameActivity.BuildConfiguration"))
-			{
-				String Configuration = bundle.getString("com.epicgames.unreal.GameActivity.BuildConfiguration");
-				IsShipping = Configuration.contains("Shipping");
-			}
 			if (bundle.containsKey("com.epicgames.unreal.GameActivity.bUseExternalFilesDir"))
 			{
 				UseExternalFilesDir = bundle.getBoolean("com.epicgames.unreal.GameActivity.bUseExternalFilesDir");
diff --git a/Engine/Plugins/Runtime/Firebase/Source/Java/notifications/EpicFirebaseMessagingService.java b/Engine/Plugins/Runtime/Firebase/Source/Java/notifications/EpicFirebaseMessagingService.java
--- a/Engine/Plugins/Runtime/Firebase/Source/Java/notifications/EpicFirebaseMessagingService.java
+++ b/Engine/Plugins/Runtime/Firebase/Source/Java/notifications/EpicFirebaseMessagingService.java
@@ -399,7 +399,8 @@ public class EpicFirebaseMessagingService extends FirebaseMessagingService {
 			queryFirebaseToken(context, projectID);
 		}
 
-		if (!Get().nativeIsShippingBuild()) {
+		if (!GameActivity.IS_SHIPPING_CONFIG)
+		{
 			Log.debug("Firebase token is " + token);
 		}
 
diff --git a/Engine/Source/Programs/UnrealBuildTool/Platform/Android/UEDeployAndroid.cs b/Engine/Source/Programs/UnrealBuildTool/Platform/Android/UEDeployAndroid.cs
--- a/Engine/Source/Programs/UnrealBuildTool/Platform/Android/UEDeployAndroid.cs
+++ b/Engine/Source/Programs/UnrealBuildTool/Platform/Android/UEDeployAndroid.cs
@@ -5074,8 +5074,8 @@ dependencies.implementation('androidx.tracing:tracing:1.2.0')
 				WriteJavaOBBDataFile(UnrealOBBDataFileName, PackageName, RequiredOBBFiles, CookFlavor, bPackageDataInsideApk, Arch);
 
 				// update GameActivity.java and GameApplication.java if out of date
-				UpdateGameActivity(Arch, NDKArch, EngineDirectory, UnrealBuildPath, bDontBundleLibrariesInAPK, bJavaPerfettoMarkers);
-				UpdateGameApplication(Arch, NDKArch, EngineDirectory, UnrealBuildPath, bJavaPerfettoMarkers);
+				UpdateGameActivity(Arch, NDKArch, EngineDirectory, UnrealBuildPath, bDontBundleLibrariesInAPK, bJavaPerfettoMarkers, Configuration);
+				UpdateGameApplication(Arch, NDKArch, EngineDirectory, UnrealBuildPath, bJavaPerfettoMarkers, Configuration);
 
 				// we don't actually need the SO for the bSkipGradleBuild case
 				string? SourceSOName = null;
@@ -6088,7 +6088,7 @@ popd
 			return "empty";
 		}
 
-		private void UpdateGameActivity(UnrealArch UnrealArch, string NDKArch, string EngineDir, string UnrealBuildPath, bool bDontBundleLibrariesInAPK, bool bJavaPerfettoMarkers)
+		private void UpdateGameActivity(UnrealArch UnrealArch, string NDKArch, string EngineDir, string UnrealBuildPath, bool bDontBundleLibrariesInAPK, bool bJavaPerfettoMarkers, UnrealTargetConfiguration Configuration)
 		{
 			string SourceFilename = Path.Combine(EngineDir, "Build", "Android", "Java", "src", "com", "epicgames", "unreal", "GameActivity.java.template");
 			string DestFilename = Path.Combine(UnrealBuildPath, "src", "com", "epicgames", "unreal", "GameActivity.java");
@@ -6222,6 +6222,7 @@ import java.util.Collection;
 				{ "//$${soLoadUnreal}$$", UPL.ProcessPluginNode(NDKArch, "soLoadUnreal", LoadUnrealDefault)},
 				{ "//$${soLoadLibrary}$$", UPL.ProcessPluginNode(NDKArch, "soLoadLibrary", LoadLibraryDefaults)},
 				{ "$${gameActivitySuperClass}$$", SuperClassDefault},
+				{ "$${gameActivityIsShippingConfig}$$", Configuration == UnrealTargetConfiguration.Shipping ? "true" : "false"},
 			};
 
 			string[] TemplateSrc = File.ReadAllLines(SourceFilename);
@@ -6304,7 +6305,7 @@ import java.util.Collection;
 			}
 		}
 
-		private void UpdateGameApplication(UnrealArch UnrealArch, string NDKArch, string EngineDir, string UnrealBuildPath, bool bJavaPerfettoMarkers)
+		private void UpdateGameApplication(UnrealArch UnrealArch, string NDKArch, string EngineDir, string UnrealBuildPath, bool bJavaPerfettoMarkers, UnrealTargetConfiguration Configuration)
 		{
 			string SourceFilename = Path.Combine(EngineDir, "Build", "Android", "Java", "src", "com", "epicgames", "unreal", "GameApplication.java.template");
 			string DestFilename = Path.Combine(UnrealBuildPath, "src", "com", "epicgames", "unreal", "GameApplication.java");
@@ -6323,6 +6324,7 @@ import java.util.Collection;
 				{ "//$${gameApplicationOnTrimMemoryAdditions}$$", UPL.ProcessPluginNode(NDKArch, "gameApplicationOnTrimMemoryAdditions", "")},
 				{ "//$${gameApplicationOnConfigurationChangedAdditions}$$", UPL.ProcessPluginNode(NDKArch, "gameApplicationOnConfigurationChangedAdditions", "")},
 				{ "//$${gameApplicationClassAdditions}$$", UPL.ProcessPluginNode(NDKArch, "gameApplicationClassAdditions", "")},
+				{ "$${gameApplicationIsShippingConfig}$$", Configuration == UnrealTargetConfiguration.Shipping ? "true" : "false"},
 			};
 
 			string[] TemplateSrc = File.ReadAllLines(SourceFilename);
