Fix keyboard EditBox layout on devices with cutout when the cutout area isn't used. getWindowVisibleDisplayFrame returns a rectable in screen space, however, the edit box coordiates need to be provided in the main decor view's space. As such, translate the results of getWindowVisibleDisplayFrame to the decor view's space by translating these bu the decor view's screen position.

diff --git a/Engine/Build/Android/Java/src/com/epicgames/unreal/GameActivity.java.template b/Engine/Build/Android/Java/src/com/epicgames/unreal/GameActivity.java.template
--- a/Engine/Build/Android/Java/src/com/epicgames/unreal/GameActivity.java.template
+++ b/Engine/Build/Android/Java/src/com/epicgames/unreal/GameActivity.java.template
@@ -4752,6 +4752,14 @@ Trace.beginSection("UI Building");
 
         				visibleView.getWindowVisibleDisplayFrame( visibleRect );
 
+						// move the visible rect from screen space to mainDecorView's space
+						int [] mainDecorViewLocation = new int[2];
+						mainDecorView.getLocationOnScreen(mainDecorViewLocation);
+						visibleRect.right -= mainDecorViewLocation[0];
+						visibleRect.left -= mainDecorViewLocation[0];
+						visibleRect.top -= mainDecorViewLocation[1];
+						visibleRect.bottom -= mainDecorViewLocation[1];
+						
 						mainDecorView.getDrawingRect( mainDecorViewRect );
 
 						//Log.debug("VK: onGlobalLayout visibleRect:(" +  visibleRect.left + ", " + visibleRect.top +", " + visibleRect.right +", " + visibleRect.bottom +")"+
