Implement support for factoring screen rounded edges in computing the safe zone. The implementation allows user to align the safe zone horizontally or vertically to the rounded edge excluded areas or to an average of both. Additionally, clean up current safe zone implementation to determine size relative to the main view instead of the display as the system may reduce the main view to exclude cutout areas.

diff --git a/Engine/Build/Android/Java/src/com/epicgames/unreal/GameActivity.java.template b/Engine/Build/Android/Java/src/com/epicgames/unreal/GameActivity.java.template
--- a/Engine/Build/Android/Java/src/com/epicgames/unreal/GameActivity.java.template
+++ b/Engine/Build/Android/Java/src/com/epicgames/unreal/GameActivity.java.template
@@ -148,6 +148,7 @@ import android.media.AudioManager;
 import android.util.DisplayMetrics;
 
 import android.view.DisplayCutout;
+import android.view.RoundedCorner;
 import android.view.InputDevice;
 import android.view.KeyEvent;
 import android.view.Gravity;
@@ -279,7 +280,7 @@ public class GameActivity extends $${gameActivitySuperClass}$$ implements Surfac
 
 	private Rect safezone = new Rect(0, 0, 0, 0);
 	private RectF safezoneF = new RectF(0.0f, 0.0f, 0.0f, 0.0f);
-
+	
 	private final float[] rotationMatrix = new float[9];
 	private final float[] orientationAngles = new float[3];
 
@@ -527,6 +528,8 @@ public class GameActivity extends $${gameActivitySuperClass}$$ implements Surfac
 	
 	private boolean SplashScreenLaunch = false;
 	private boolean UseDisplayCutout = false;
+	private enum RoundedEdgeSafeZoneDirection { None, Horizontal, Vertical, Both };
+	private RoundedEdgeSafeZoneDirection RoundedEdgeSafeZone = RoundedEdgeSafeZoneDirection.None;
 	private boolean ShouldHideUI = false;
 
 	/** Whether this application is for distribution */
@@ -4429,6 +4432,16 @@ Trace.endSection();
 				}
 			}
 
+			if (bundle.containsKey("com.epicgames.unreal.GameActivity.RoundedEdgeSafeZone"))
+			{
+				RoundedEdgeSafeZone = RoundedEdgeSafeZoneDirection.valueOf( bundle.getString("com.epicgames.unreal.GameActivity.RoundedEdgeSafeZone") );
+				Log.debug( "RoundedEdgeSafeZone set to " + RoundedEdgeSafeZone);
+			}
+			else
+			{
+				Log.debug( "RoundedEdgeSafeZone not found. Leaving as " + RoundedEdgeSafeZone);
+			}
+
 			if (bundle.containsKey("com.epicgames.unreal.GameActivity.bPublicLogFiles"))
             {
                 PublicLogFiles = bundle.getBoolean("com.epicgames.unreal.GameActivity.bPublicLogFiles");
@@ -4690,7 +4703,7 @@ Trace.beginSection("UI Building");
 
 			mainDecorView = getWindow().getDecorView();
 			mainDecorViewRect = new Rect();
-		
+
 			// check for display cutouts to set safe zone
 			View mainRootView = mainDecorView.findViewById(android.R.id.content).getRootView();
 			if (Build.VERSION.SDK_INT > 19)
@@ -4698,28 +4711,19 @@ Trace.beginSection("UI Building");
 				mainRootView.setOnApplyWindowInsetsListener(new View.OnApplyWindowInsetsListener() {
 					@Override
        				public WindowInsets onApplyWindowInsets(View v, WindowInsets insets) {
+						safezone.set(0,0,0,0);
+		
 						if (Build.VERSION.SDK_INT >= 28)
 						{
 							DisplayCutout cutout = insets.getDisplayCutout();
+				
 							if (cutout != null)
 							{
-								Rect fullViewRect = new Rect();
-								mainDecorView.getDrawingRect( mainDecorViewRect );
-
-								android.view.Display display = getWindowManager().getDefaultDisplay();
-								android.graphics.Point displaySize = new android.graphics.Point();
-								display.getRealSize(displaySize);
-
-								// store safezone independent of resolution
-								float invX = 1.0f / (displaySize.x != 0.0f ? displaySize.x : 1.0f);
-								float invY = 1.0f / (displaySize.y != 0.0f ? displaySize.y : 1.0f);
-								safezoneF.left = (safezone.left = cutout.getSafeInsetLeft()) * invX;
-								safezoneF.top = (safezone.top = cutout.getSafeInsetTop()) * invY;
-								safezoneF.right = (safezone.right = cutout.getSafeInsetRight()) * invX;
-								safezoneF.bottom = (safezone.bottom = cutout.getSafeInsetBottom()) *invY;
-								nativeSetSafezoneInfo(displaySize.y > displaySize.x, safezoneF.left, safezoneF.top, safezoneF.right, safezoneF.bottom);
-
-								//Log.debug("FullViewRect: " + fullViewRect.toString());	
+								safezone.left = cutout.getSafeInsetLeft();
+								safezone.top = cutout.getSafeInsetTop();
+								safezone.right = cutout.getSafeInsetRight();
+								safezone.bottom = cutout.getSafeInsetBottom();
+								
 								//Log.debug("DisplayCutout: " + cutout.toString());
 								//Log.debug("SafeTop: " + cutout.getSafeInsetTop());
 								//Log.debug("SafeBottom: " + cutout.getSafeInsetBottom());
@@ -4731,11 +4735,45 @@ Trace.beginSection("UI Building");
 								//	Log.debug("Rect: " + r.toString());
 								//}
 							}
+					
+							if (Build.VERSION.SDK_INT >= 31 && RoundedEdgeSafeZone != RoundedEdgeSafeZoneDirection.None)
+							{
+								final RoundedCorner topLeft = insets.getRoundedCorner(RoundedCorner.POSITION_TOP_LEFT);
+								final RoundedCorner bottomRight = insets.getRoundedCorner(RoundedCorner.POSITION_BOTTOM_RIGHT);
+								
+								if (topLeft != null && bottomRight != null)
+								{
+									float bothMult = RoundedEdgeSafeZone == RoundedEdgeSafeZoneDirection.Both ? 0.5f : 0.0f; 
+									float horizMult = RoundedEdgeSafeZone == RoundedEdgeSafeZoneDirection.Horizontal ? 1.0f : bothMult;
+									float vertMult = RoundedEdgeSafeZone == RoundedEdgeSafeZoneDirection.Vertical ? 1.0f : bothMult;
+									int leftInset = (int)(horizMult * topLeft.getRadius());
+									int topInset = (int)(vertMult * topLeft.getRadius());
+									int rightInset = (int)(horizMult * bottomRight.getRadius());
+									int bottomInset = (int)(vertMult * bottomRight.getRadius());
+									
+									safezone.set( Math.max( safezone.left, leftInset ), Math.max( safezone.top, topInset ), Math.max( safezone.right, rightInset ), Math.max( safezone.bottom, bottomInset ) );
+								}
+							}
+						}
+				
+						// store safezone independent of resolution
+						android.graphics.Point displaySize = new android.graphics.Point();
+						displaySize.set( mainDecorViewRect.right - mainDecorViewRect.left, mainDecorViewRect.bottom - mainDecorViewRect.top );
+						if (displaySize.x == 0 || displaySize.y == 0)
+						{
+							android.view.Display display = getWindowManager().getDefaultDisplay();
+							display.getRealSize(displaySize);
 						}
+						float invX = 1.0f / displaySize.x;
+						float invY = 1.0f / displaySize.y;
+						safezoneF.set(safezone.left * invX, safezone.top * invY, safezone.right * invX, safezone.bottom * invY);
+						nativeSetSafezoneInfo(displaySize.y > displaySize.x, safezoneF.left, safezoneF.top, safezoneF.right, safezoneF.bottom);
+
 						return insets;
 					}
 				});
 			}
+			
 		
 			createVirtualKeyboardInput();
 		
@@ -4831,6 +4869,8 @@ Trace.beginSection("UI Building");
 							AndroidThunkJava_ResizeKeyboard(false, keyboardRect, visibleRect);
 						}
         			}
+        			
+        			mainView.requestApplyInsets();
         		}
 			});
 		}
diff --git a/Engine/Config/BaseEngine.ini b/Engine/Config/BaseEngine.ini
--- a/Engine/Config/BaseEngine.ini
+++ b/Engine/Config/BaseEngine.ini
@@ -3317,6 +3317,7 @@ bEnableMulticastSupport=False
 bEnableRayTracing=true
 bAllowResizing=false
 bSupportSizeChanges=false
+RoundedEdgeSafeZone=None
 
 [/Script/AndroidPlatformEditor.AndroidSDKSettings]
 SDKAPILevel=latest
diff --git a/Engine/Source/Programs/UnrealBuildTool/Platform/Android/UEDeployAndroid.cs b/Engine/Source/Programs/UnrealBuildTool/Platform/Android/UEDeployAndroid.cs
--- a/Engine/Source/Programs/UnrealBuildTool/Platform/Android/UEDeployAndroid.cs
+++ b/Engine/Source/Programs/UnrealBuildTool/Platform/Android/UEDeployAndroid.cs
@@ -2718,6 +2718,8 @@ namespace UnrealBuildTool
 			Ini.GetBool("/Script/AndroidRuntimeSettings.AndroidRuntimeSettings", "bFullScreen", out EnableFullScreen);
 			bool bUseDisplayCutout;
 			Ini.GetBool("/Script/AndroidRuntimeSettings.AndroidRuntimeSettings", "bUseDisplayCutout", out bUseDisplayCutout);
+			string RoundedEdgeSafeZone;
+			Ini.GetString("/Script/AndroidRuntimeSettings.AndroidRuntimeSettings", "RoundedEdgeSafeZone", out RoundedEdgeSafeZone);
 			bool bAllowResizing;
 			Ini.GetBool("/Script/AndroidRuntimeSettings.AndroidRuntimeSettings", "bAllowResizing", out bAllowResizing);
 			bool bSupportSizeChanges;
@@ -3060,6 +3062,7 @@ namespace UnrealBuildTool
 			Text.AppendLine(String.Format("\t\t<meta-data android:name=\"com.epicgames.unreal.GameActivity.bUseExternalFilesDir\" android:value=\"{0}\"/>", bUseExternalFilesDir ? "true" : "false"));
 			Text.AppendLine(String.Format("\t\t<meta-data android:name=\"com.epicgames.unreal.GameActivity.bPublicLogFiles\" android:value=\"{0}\"/>", bPublicLogFiles ? "true" : "false"));
 			Text.AppendLine(String.Format("\t\t<meta-data android:name=\"com.epicgames.unreal.GameActivity.bUseDisplayCutout\" android:value=\"{0}\"/>", bUseDisplayCutout ? "true" : "false"));
+			Text.AppendLine(String.Format("\t\t<meta-data android:name=\"com.epicgames.unreal.GameActivity.RoundedEdgeSafeZone\" android:value=\"{0}\"/>", RoundedEdgeSafeZone));
 			Text.AppendLine(String.Format("\t\t<meta-data android:name=\"com.epicgames.unreal.GameActivity.bAllowIMU\" android:value=\"{0}\"/>", bAllowIMU ? "true" : "false"));
 			Text.AppendLine(String.Format("\t\t<meta-data android:name=\"com.epicgames.unreal.GameActivity.bSupportsVulkan\" android:value=\"{0}\"/>", bSupportsVulkan ? "true" : "false"));
 			Text.AppendLine(String.Format("\t\t<meta-data android:name=\"com.epicgames.unreal.GameActivity.PropagateAlpha\" android:value=\"{0}\"/>", PropagateAlpha));
diff --git a/Engine/Source/Runtime/Android/AndroidRuntimeSettings/Classes/AndroidRuntimeSettings.h b/Engine/Source/Runtime/Android/AndroidRuntimeSettings/Classes/AndroidRuntimeSettings.h
--- a/Engine/Source/Runtime/Android/AndroidRuntimeSettings/Classes/AndroidRuntimeSettings.h
+++ b/Engine/Source/Runtime/Android/AndroidRuntimeSettings/Classes/AndroidRuntimeSettings.h
@@ -45,6 +45,19 @@ namespace EAndroidScreenOrientation
 	};
 }
 
+/** Rounded edge safe zone preferences */
+UENUM()
+namespace EAndroidRoundedEdgeSafeZoneDirection
+{
+	enum Type : int
+	{
+		None,
+		Horizontal,
+		Vertical,
+		Both
+	};
+}
+
 /** Depth buffer precision preferences */
 UENUM()
 namespace EAndroidDepthBufferPreference
@@ -291,6 +304,10 @@ public:
 	UPROPERTY(GlobalConfig, EditAnywhere, Category = "APK Packaging", Meta = (DisplayName = "Use display cutout region?"))
 	bool bUseDisplayCutout;
 
+	// Enables limiting safe zone away from rounded edges on Android 12+
+	UPROPERTY(GlobalConfig, EditAnywhere, Category = "APK Packaging", Meta = (DisplayName = "Limit safe zone to screen rounded edges."))
+	TEnumAsByte<EAndroidRoundedEdgeSafeZoneDirection::Type> RoundedEdgeSafeZone;
+	
 	// Allow resizing of the window on Android devices with splitscreen
 	UPROPERTY(GlobalConfig, EditAnywhere, Category = "APK Packaging", Meta = (DisplayName = "Allow splitscreen resizing?"))
 	bool bAllowResizing;
diff --git a/Engine/Source/Runtime/ApplicationCore/Private/Android/AndroidWindow.cpp b/Engine/Source/Runtime/ApplicationCore/Private/Android/AndroidWindow.cpp
--- a/Engine/Source/Runtime/ApplicationCore/Private/Android/AndroidWindow.cpp
+++ b/Engine/Source/Runtime/ApplicationCore/Private/Android/AndroidWindow.cpp
@@ -162,8 +162,7 @@ void FAndroidWindow::SetOSWindowHandle(void* InWindow)
 static bool GAndroidIsPortrait = false;
 static EDeviceScreenOrientation GDeviceScreenOrientation = EDeviceScreenOrientation::Unknown;
 static int GAndroidDepthBufferPreference = 0;
-static FVector4 GAndroidPortraitSafezone = FVector4(-1.0f, -1.0f, -1.0f, -1.0f);
-static FVector4 GAndroidLandscapeSafezone = FVector4(-1.0f, -1.0f, -1.0f, -1.0f);
+static FVector4 GAndroidPortraitSafezone(-1.0f, -1.0f, -1.0f, -1.0f), GAndroidLandscapeSafezone(-1.0f, -1.0f, -1.0f, -1.0f), *GAndroidLastSafezone = nullptr;
 #if USE_ANDROID_JNI
 JNI_METHOD void Java_com_epicgames_unreal_GameActivity_nativeSetWindowInfo(JNIEnv* jenv, jobject thiz, jboolean bIsPortrait, jint DepthBufferPreference, jint PropagateAlpha)
 {
@@ -190,22 +189,17 @@ JNI_METHOD void Java_com_epicgames_unreal_GameActivity_nativeSetSurfaceViewInfo(
 
 JNI_METHOD void Java_com_epicgames_unreal_GameActivity_nativeSetSafezoneInfo(JNIEnv* jenv, jobject thiz, jboolean bIsPortrait, jfloat left, jfloat top, jfloat right, jfloat bottom)
 {
-	if (bIsPortrait)
-	{
-		GAndroidPortraitSafezone.X = left;
-		GAndroidPortraitSafezone.Y = top;
-		GAndroidPortraitSafezone.Z = right;
-		GAndroidPortraitSafezone.W = bottom;
-	}
-	else
+	FVector4 NewSafezone(left, top, right, bottom);
+	FVector4& TargetSafezone = bIsPortrait ? GAndroidPortraitSafezone : GAndroidLandscapeSafezone;
+	bool SafezoneUpdated = TargetSafezone != NewSafezone || GAndroidLastSafezone != &TargetSafezone;
+	GAndroidLastSafezone = &TargetSafezone;
+	TargetSafezone = NewSafezone;
+	
+#if USE_ANDROID_EVENTS
+	if( SafezoneUpdated )
 	{
-		GAndroidLandscapeSafezone.X = left;
-		GAndroidLandscapeSafezone.Y = top;
-		GAndroidLandscapeSafezone.Z = right;
-		GAndroidLandscapeSafezone.W = bottom;
+		FAppEventManager::GetInstance()->EnqueueAppEvent(APP_EVENT_STATE_SAFE_ZONE_UPDATED);
 	}
-#if USE_ANDROID_EVENTS
-	FAppEventManager::GetInstance()->EnqueueAppEvent(APP_EVENT_STATE_SAFE_ZONE_UPDATED);
 #endif
 	UE_LOG(LogAndroid, Log, TEXT("nativeSetSafezoneInfo bIsPortrait=%d, left=%f, top=%f, right=%f, bottom=%f"), bIsPortrait ? 1 : 0, left, top, right, bottom);
 }
