Android: Ensure native configrules are always set.
Fixes a potential infinite wait on native side configrule access if the java code had failed to process them.

#tests reprod bug + confirm fix on local client

diff --git a/Engine/Build/Android/Java/src/com/epicgames/unreal/GameActivity.java.template b/Engine/Build/Android/Java/src/com/epicgames/unreal/GameActivity.java.template
--- a/Engine/Build/Android/Java/src/com/epicgames/unreal/GameActivity.java.template
+++ b/Engine/Build/Android/Java/src/com/epicgames/unreal/GameActivity.java.template
@@ -383,6 +383,7 @@ public class GameActivity extends $${gameActivitySuperClass}$$ implements Surfac
 	protected long ConfigRulesCRC = 0;
 	protected long ConfigRulesVarsCRC = 0;
 	protected Map<String, String> ConfigRulesVars = null;
+	protected boolean bConfigRulesApplied = false;
 	protected LinkedList<String> SystemInfoLogs = null;
 	protected Map<String, String> CpuInfo = null;
 
@@ -3646,6 +3647,7 @@ Trace.endSection();  // WriteCache
 			keyValues[keyValueIndex++] = entry.getValue();
 		}
 		nativeSetConfigRulesVariables(keyValues);
+		bConfigRulesApplied = true;
 	}
 
 	public void LoggerCallback(String Level, String Tag, String Message)
@@ -3834,6 +3836,13 @@ Trace.endSection();
 		super.onCreate(savedInstanceState);
 
 		onCreateBody(savedInstanceState);
+		if (!bConfigRulesApplied)
+		{
+			SystemInfo_DumpLog();
+			// the native code will wait for configrules variables.
+			// something must be sent to release it.
+			nativeSetConfigRulesVariables(new String[]{});
+		}
 	}
 
 	protected void onCreatePreBody()
