Resolves potential sync issue that leaves the EditBox non responsive to text input and the keyboard in a state where it doesn't get dismissed its the'Done' button. Additionally, replace setVisibility with add/remove View to work around an Android level spam issue when keyboard is repeatedly dismissed/invoked.

diff --git a/Engine/Build/Android/Java/src/com/epicgames/unreal/GameActivity.java.template b/Engine/Build/Android/Java/src/com/epicgames/unreal/GameActivity.java.template
--- a/Engine/Build/Android/Java/src/com/epicgames/unreal/GameActivity.java.template
+++ b/Engine/Build/Android/Java/src/com/epicgames/unreal/GameActivity.java.template
@@ -4853,8 +4853,6 @@ Trace.beginSection("UI Building");
 								newVirtualKeyboardInput.getLayoutParams().width = Math.abs(visibleRect.right - visibleRect.left);
 								newVirtualKeyboardInput.setX(leftDiff);
 								newVirtualKeyboardInput.setY(keyboardYPos);
-								newVirtualKeyboardInput.setVisibility(View.VISIBLE);
-								newVirtualKeyboardInput.requestFocus();
 							}
 
 							AndroidThunkJava_ResizeKeyboard(true, keyboardRect, visibleRect);
@@ -4869,7 +4867,7 @@ Trace.beginSection("UI Building");
 								{
 
 									//Log.debug("VK: hide");
-									newVirtualKeyboardInput.setVisibility(View.GONE);
+
 									//set offscreen
 									newVirtualKeyboardInput.setY(-1000);
 								}
@@ -5796,7 +5794,7 @@ Trace.endSection();
 	}
 
 	// new functions to show/hide virtual keyboard
-	public void AndroidThunkJava_HideVirtualKeyboardInput()
+	public synchronized void AndroidThunkJava_HideVirtualKeyboardInput()
 	{
 		//Log.debug("VK: AndroidThunkJava_HideVirtualKeyboardInput");
 
@@ -5818,7 +5816,7 @@ Trace.endSection();
 	//initial settings for the virtual input
 	String virtualKeyboardInputContent;
 	int virtualKeyboardInputType;
-	public void AndroidThunkJava_ShowVirtualKeyboardInput(int inInputType, String Label, String Contents)
+	public synchronized void AndroidThunkJava_ShowVirtualKeyboardInput(int inInputType, String Label, String Contents)
 	{
 		//Log.debug("VK: AndroidThunkJava_ShowVirtualKeyboardInput");
 		virtualKeyboardInputContent = Contents;
@@ -5840,123 +5838,127 @@ Trace.endSection();
 	
 	//jira UE-49141 Virtual keyboard is unresponsive with repeated tapping in control (some devices)
 	//#jira UE-49139 Tapping in the same text box doesn't make the virtual keyboard disappear
-	public void processLastVirtualKeyboardCommand()
+	public synchronized void processLastVirtualKeyboardCommand()
 	{
 		Log.debug("VK: process last command " + lastVirtualKeyboardCommand);
-		synchronized(this) {
-			switch(lastVirtualKeyboardCommand)	
+		switch(lastVirtualKeyboardCommand)	
+		{
+			case VK_CMD_SHOW:
 			{
-				case VK_CMD_SHOW:
+				if (virtualKeyboardLayout.getChildCount() == 0)
 				{
-					newVirtualKeyboardInput.setVisibility(View.VISIBLE);
+					virtualKeyboardLayout.addView(newVirtualKeyboardInput);
+				}
 				
-					//newVirtualKeyboardInput.setBackgroundColor(Color.TRANSPARENT);
-					//newVirtualKeyboardInput.setCursorVisible(false);
+				//newVirtualKeyboardInput.setBackgroundColor(Color.TRANSPARENT);
+				//newVirtualKeyboardInput.setCursorVisible(false);
 
-					//set offscreen
-					newVirtualKeyboardInput.setY(-1000);
+				//set offscreen
+				newVirtualKeyboardInput.setY(-1000);
 
-					//set new content
-					newVirtualKeyboardInput.setText(virtualKeyboardInputContent);
-				
-					int newVirtualKeyboardInputType = virtualKeyboardInputType;
-
-					//commented: as it will disable text prediction for all devices, 
-					//	for most of them it will also block the VK in latin/english subtype
-					//	disabling chinese or korean subtypes
-					//if((virtualKeyboardInputType & InputType.TYPE_TEXT_VARIATION_PASSWORD) == 0)
-					//{
-						//#jira UE-49117 Chinese and Korean virtual keyboards don't allow native characters
-						//#jira UE-49121 Gboard and Swift swipe entry are not supported by Virtual keyboard
-						//newVirtualKeyboardInputType |= TYPE_TEXT_VARIATION_VISIBLE_PASSWORD;
-					//}
-
-					//TYPE: disable text suggestion
-					newVirtualKeyboardInputType |= InputType.TYPE_TEXT_FLAG_NO_SUGGESTIONS;
-				
-					//TYPE: set input type flags
-					newVirtualKeyboardInput.setInputType(newVirtualKeyboardInputType);
-					newVirtualKeyboardInput.setRawInputType(newVirtualKeyboardInputType);
+				//set new content
+				newVirtualKeyboardInput.setText(virtualKeyboardInputContent);
+			
+				int newVirtualKeyboardInputType = virtualKeyboardInputType;
+
+				//commented: as it will disable text prediction for all devices, 
+				//	for most of them it will also block the VK in latin/english subtype
+				//	disabling chinese or korean subtypes
+				//if((virtualKeyboardInputType & InputType.TYPE_TEXT_VARIATION_PASSWORD) == 0)
+				//{
+					//#jira UE-49117 Chinese and Korean virtual keyboards don't allow native characters
+					//#jira UE-49121 Gboard and Swift swipe entry are not supported by Virtual keyboard
+					//newVirtualKeyboardInputType |= TYPE_TEXT_VARIATION_VISIBLE_PASSWORD;
+				//}
+
+				//TYPE: disable text suggestion
+				newVirtualKeyboardInputType |= InputType.TYPE_TEXT_FLAG_NO_SUGGESTIONS;
+			
+				//TYPE: set input type flags
+				newVirtualKeyboardInput.setInputType(newVirtualKeyboardInputType);
+				newVirtualKeyboardInput.setRawInputType(newVirtualKeyboardInputType);
 
-					//IME: set Done button for single line input
-					int imeOptions = EditorInfo.IME_FLAG_NO_EXTRACT_UI;
+				//IME: set Done button for single line input
+				int imeOptions = EditorInfo.IME_FLAG_NO_EXTRACT_UI;
 
 
-					if (ANDROID_BUILD_VERSION >=  11)
-					{
-						imeOptions |= EditorInfo.IME_FLAG_NO_FULLSCREEN;
-					}
+				if (ANDROID_BUILD_VERSION >=  11)
+				{
+					imeOptions |= EditorInfo.IME_FLAG_NO_FULLSCREEN;
+				}
 
-					//IME: set single/multi line input type
-					if((virtualKeyboardInputType & InputType.TYPE_TEXT_FLAG_MULTI_LINE) != 0)
-					{
-						//disable enter for multi-line - will be treated by virtualKeyboardInputType in sendKeyEvent
-						newVirtualKeyboardInput.setSingleLine(false);
-						//#jira UE-49128 Virtual Keyboard text field doesn't appear if there is too much text
-						newVirtualKeyboardInput.setMaxLines(5);
-						imeOptions |= EditorInfo.IME_FLAG_NO_ENTER_ACTION;
-						imeOptions &= ~EditorInfo.IME_ACTION_DONE;
-					}
-					else
-					{
-						newVirtualKeyboardInput.setSingleLine(true);
-						newVirtualKeyboardInput.setMaxLines(1);
-						imeOptions &= ~EditorInfo.IME_FLAG_NO_ENTER_ACTION;
-						imeOptions |= EditorInfo.IME_ACTION_DONE;
-					}
+				//IME: set single/multi line input type
+				if((virtualKeyboardInputType & InputType.TYPE_TEXT_FLAG_MULTI_LINE) != 0)
+				{
+					//disable enter for multi-line - will be treated by virtualKeyboardInputType in sendKeyEvent
+					newVirtualKeyboardInput.setSingleLine(false);
+					//#jira UE-49128 Virtual Keyboard text field doesn't appear if there is too much text
+					newVirtualKeyboardInput.setMaxLines(5);
+					imeOptions |= EditorInfo.IME_FLAG_NO_ENTER_ACTION;
+					imeOptions &= ~EditorInfo.IME_ACTION_DONE;
+				}
+				else
+				{
+					newVirtualKeyboardInput.setSingleLine(true);
+					newVirtualKeyboardInput.setMaxLines(1);
+					imeOptions &= ~EditorInfo.IME_FLAG_NO_ENTER_ACTION;
+					imeOptions |= EditorInfo.IME_ACTION_DONE;
+				}
 
-					//IME: set IME flags
-					newVirtualKeyboardInput.setImeOptions(imeOptions);
+				//IME: set IME flags
+				newVirtualKeyboardInput.setImeOptions(imeOptions);
 
-					//TRANSFORMATION: hide input for passwords
-					newVirtualKeyboardInput.setTransformationMethod((
-						virtualKeyboardInputType & InputType.TYPE_TEXT_VARIATION_PASSWORD) == 0 ? 
-						null : 
-						PasswordTransformationMethod.getInstance());
+				//TRANSFORMATION: hide input for passwords
+				newVirtualKeyboardInput.setTransformationMethod((
+					virtualKeyboardInputType & InputType.TYPE_TEXT_VARIATION_PASSWORD) == 0 ? 
+					null : 
+					PasswordTransformationMethod.getInstance());
 
-					//SELECTION: move to end
-					newVirtualKeyboardInput.setSelection(newVirtualKeyboardInput.getText().length());
+				//SELECTION: move to end
+				newVirtualKeyboardInput.setSelection(newVirtualKeyboardInput.getText().length());
 
-					if(newVirtualKeyboardInput.requestFocus())
-					{
-						//Log.debug("VK: Show newVirtualKeyboardInput");
-						InputMethodManager imm = (InputMethodManager) getSystemService(Context.INPUT_METHOD_SERVICE);
-						imm.showSoftInput(newVirtualKeyboardInput, 0);
+				if (newVirtualKeyboardInput.requestFocus())
+				{
+					//Log.debug("VK: Show newVirtualKeyboardInput");
+					InputMethodManager imm = (InputMethodManager) getSystemService(Context.INPUT_METHOD_SERVICE);
+					imm.showSoftInput(newVirtualKeyboardInput, 0);
 
-						nativeVirtualKeyboardVisible(true);
-						bKeyboardShowing = true;
+					nativeVirtualKeyboardVisible(true);
+					bKeyboardShowing = true;
 
-						CheckKeyboardDisplayed();
-					}
+					CheckKeyboardDisplayed();
 				}
-				break;
-				case VK_CMD_HIDE:
+			}
+			break;
+			case VK_CMD_HIDE:
+			{
+				if (bKeyboardShowing || newVirtualKeyboardInput.getVisibility() != View.GONE)
 				{
-						if (bKeyboardShowing || newVirtualKeyboardInput.getVisibility() != View.GONE)
-						{
-							//Log.debug("VK: Hide newVirtualKeyboardInput");
+					//Log.debug("VK: Hide newVirtualKeyboardInput");
 
-							if (newVirtualKeyboardInput.hasFocus()) {
-								try {
-									newVirtualKeyboardInput.clearFocus();
-								} catch (Exception e) {
-									Log.warn("Unable to clear focus from virtualKeyboardInput");
-								}
-							}
-							//set offscreen
-							newVirtualKeyboardInput.setY(-1000);
+					if (newVirtualKeyboardInput.hasFocus()) {
+						try {
+							newVirtualKeyboardInput.clearFocus();
+						} catch (Exception e) {
+							Log.warn("Unable to clear focus from virtualKeyboardInput");
+						}
+					}
+					//set offscreen
+					newVirtualKeyboardInput.setY(-1000);
 
-							newVirtualKeyboardInput.setVisibility(View.GONE);
+					if (virtualKeyboardLayout.getChildCount() == 1)
+					{
+						virtualKeyboardLayout.removeView(newVirtualKeyboardInput);
+					}
 
-							InputMethodManager imm =(InputMethodManager)getSystemService(Context.INPUT_METHOD_SERVICE);
-							imm.hideSoftInputFromWindow(newVirtualKeyboardInput.getWindowToken(), 0);
+					InputMethodManager imm =(InputMethodManager)getSystemService(Context.INPUT_METHOD_SERVICE);
+					imm.hideSoftInputFromWindow(newVirtualKeyboardInput.getWindowToken(), 0);
 
-							nativeVirtualKeyboardVisible(false);
-							bKeyboardShowing = false;
-						}
+					nativeVirtualKeyboardVisible(false);
+					bKeyboardShowing = false;
 				}
-				break;
 			}
+			break;
 		}
 		lastVirtualKeyboardCommand = VirtualKeyboardCommand.VK_CMD_NONE;
 	}
@@ -5968,7 +5970,7 @@ Trace.endSection();
 		{
 			public void run() 
 			{
-				synchronized(this) 
+				synchronized(_activity) 
 				{
 					if(bKeyboardShowing && newVirtualKeyboardInput.getY() < 0) 
 					{
@@ -8105,7 +8107,6 @@ Trace.endSection();
 		newVirtualKeyboardInput.setBackgroundColor(Color.WHITE);
 		newVirtualKeyboardInput.setTextColor(Color.BLACK);
 		newVirtualKeyboardInput.setLayoutParams(new RelativeLayout.LayoutParams(LayoutParams.MATCH_PARENT, LayoutParams.WRAP_CONTENT));
-		newVirtualKeyboardInput.setVisibility(View.GONE);
 		if (ANDROID_BUILD_VERSION < 11)
 		{
 			newVirtualKeyboardInput.setImeOptions(EditorInfo.IME_FLAG_NO_EXTRACT_UI);
@@ -8237,8 +8238,6 @@ Trace.endSection();
 			}
 		});
 		
-		virtualKeyboardLayout.addView (newVirtualKeyboardInput);
-
 		virtualKeyboardHandler = new Handler(android.os.Looper.getMainLooper());
 
 	}
