Fix to small JNI bugs in PSOServer and AndroidJavaEnv

#android

diff --git a/Engine/Build/Android/Prebuilt/PSOService/Source/PSOService.cpp b/Engine/Build/Android/Prebuilt/PSOService/Source/PSOService.cpp
--- a/Engine/Build/Android/Prebuilt/PSOService/Source/PSOService.cpp
+++ b/Engine/Build/Android/Prebuilt/PSOService/Source/PSOService.cpp
@@ -73,7 +73,7 @@ JNIEXPORT jint JNI_OnLoad(JavaVM* vm, void* reserved)
 	if (!env)
 	{
 		LOG_ERROR("Failed to get Java environment");
-		exit(-1);
+		return JNI_ERR;
 	}
 
 	jclass clazz = env->FindClass(PsoServices::FPsoProgramService::ClassName);
@@ -81,16 +81,19 @@ JNIEXPORT jint JNI_OnLoad(JavaVM* vm, void* reserved)
 	if (!clazz)
 	{
 		LOG_ERROR("Failed to find Java class");
-		exit(-1);
+		return JNI_ERR;
 	}
 
-	if (env->RegisterNatives(clazz, PsoServices::FPsoProgramService::NativeMethods, sizeof(PsoServices::FPsoProgramService::NativeMethods) / sizeof(JNINativeMethod)) < 0)
+	jint result = env->RegisterNatives(clazz, PsoServices::FPsoProgramService::NativeMethods, sizeof(PsoServices::FPsoProgramService::NativeMethods) / sizeof(JNINativeMethod));
+
+	env->DeleteLocalRef(clazz);
+
+	if (result < 0)
 	{
 		LOG_ERROR("Failed to register native methods");
-		exit(-1);
+		return JNI_ERR;
 	}
 
-	env->DeleteLocalRef(clazz);
 	return JNI_VERSION_1_6;
 }
 
@@ -1066,6 +1069,23 @@ jbyteArray JNICALL UE::Jni::PsoServices::FPsoProgramService::CompileVKGFXPSO(JNI
 
 	FVulkanPSOCompiler::Get().CompileGFXPSO(VS, VSSize, PS, PSSize, PSO, PSOSize, PSOCacheDataSource, PSOCacheDataSourceSize);
 
+	if (PSOCacheDataSource)
+	{
+		jenv->ReleaseByteArrayElements(jPSOCacheDataSource, PSOCacheDataSource, JNI_ABORT);
+	}
+	if (PSO)
+	{
+		jenv->ReleaseByteArrayElements(jPSO, PSO, JNI_ABORT);
+	}
+	if (PS)
+	{
+		jenv->ReleaseByteArrayElements(jPS, PS, JNI_ABORT);
+	}
+	if (VS)
+	{
+		jenv->ReleaseByteArrayElements(jVS, VS, JNI_ABORT);
+	}
+
 	uint32_t Size = 0;
 
 	FVulkanPSOCompiler::Get().GetPSOBinary(nullptr, Size);
@@ -1115,7 +1135,7 @@ jint JNICALL UE::Jni::PsoServices::FPsoProgramService::CompileVKGFXPSOSHM(JNIEnv
 
 		LOG_VERBOSE("ParamsSharedBuffer %zu, %p ", memSize, ParamsSharedBuffer);
 
-		uint32_t CurrOffset = 0;
+		uint64_t CurrOffset = 0;
 		const uint8_t* VS = (const uint8_t*)ParamsSharedBuffer;
 		uint64_t VSSize = jVSSize;
 		CurrOffset += VSSize;
diff --git a/Engine/Source/Runtime/Core/Private/Android/AndroidJavaEnv.cpp b/Engine/Source/Runtime/Core/Private/Android/AndroidJavaEnv.cpp
--- a/Engine/Source/Runtime/Core/Private/Android/AndroidJavaEnv.cpp
+++ b/Engine/Source/Runtime/Core/Private/Android/AndroidJavaEnv.cpp
@@ -156,11 +156,11 @@ FScopedJavaObject<jobjectArray> FJavaHelper::ToJavaStringArray(JNIEnv* Env, cons
 		// FStringView of an empty FString contains a null pointer as data
 		if (UnrealStrings[Idx].GetData())
 		{
-			Env->SetObjectArrayElement(ObjectArray, Idx, Env->NewStringUTF(TCHAR_TO_UTF8(UnrealStrings[Idx].GetData())));
+			Env->SetObjectArrayElement(ObjectArray, Idx, *FScopedJavaObject{Env->NewStringUTF(TCHAR_TO_UTF8(UnrealStrings[Idx].GetData()))});
 		}
 		else
 		{		
-			Env->SetObjectArrayElement(ObjectArray, Idx, Env->NewStringUTF(""));
+			Env->SetObjectArrayElement(ObjectArray, Idx, *FScopedJavaObject{Env->NewStringUTF("")});
 		}
 	}
 	return NewScopedJavaObject(Env, ObjectArray);
diff --git a/Engine/Source/Runtime/Core/Public/Android/AndroidJava.h b/Engine/Source/Runtime/Core/Public/Android/AndroidJava.h
--- a/Engine/Source/Runtime/Core/Public/Android/AndroidJava.h
+++ b/Engine/Source/Runtime/Core/Public/Android/AndroidJava.h
@@ -45,7 +45,7 @@ public:
 
 	static FScopedJavaObject<jstring> GetJString(const FString& String);
 
-	void VerifyException();
+	static void VerifyException();
 
 protected:
 
