Catching UnsatisfiedLinkError in GameApplication.LoadLibraries in case some other ClassLoader has already loaded them

#android

diff --git a/Engine/Build/Android/Java/src/com/epicgames/unreal/GameApplication.java.template b/Engine/Build/Android/Java/src/com/epicgames/unreal/GameApplication.java.template
--- a/Engine/Build/Android/Java/src/com/epicgames/unreal/GameApplication.java.template
+++ b/Engine/Build/Android/Java/src/com/epicgames/unreal/GameApplication.java.template
@@ -125,7 +125,7 @@ public class GameApplication extends Application implements LifecycleObserver //
 		{
 //$${gameApplicationEarlyLoadAdditions}$$
 		}
-		catch (Exception e)
+		catch (Exception | UnsatisfiedLinkError e)
 		{
 			Log.debug(e.toString());
 		}
