Round rather than truncate Android getRefreshRate when converting to integer
#10205

diff --git a/Engine/Build/Android/Java/src/com/epicgames/unreal/GameActivity.java.template b/Engine/Build/Android/Java/src/com/epicgames/unreal/GameActivity.java.template
--- a/Engine/Build/Android/Java/src/com/epicgames/unreal/GameActivity.java.template
+++ b/Engine/Build/Android/Java/src/com/epicgames/unreal/GameActivity.java.template
@@ -7239,7 +7239,7 @@ Trace.endSection();
 				if (modes[i].getPhysicalHeight() == currentmode.getPhysicalHeight() &&
 					modes[i].getPhysicalWidth() == currentmode.getPhysicalWidth())
 				{
-					refreshlist.add((int)modes[i].getRefreshRate());
+					refreshlist.add(java.lang.Math.round(modes[i].getRefreshRate()));
 				}
 			}
 			if (refreshlist.size() == 0)
@@ -7298,7 +7298,7 @@ Trace.endSection();
 			{
 				if (modes[i].getPhysicalHeight() == currentmode.getPhysicalHeight() &&
 					modes[i].getPhysicalWidth() == currentmode.getPhysicalWidth() &&
-					(int)modes[i].getRefreshRate() == RefreshRate)
+					java.lang.Math.round(modes[i].getRefreshRate()) == RefreshRate)
 				{
 					final int modeid = modes[i].getModeId();
 					if(currentmodeid != modeid)
@@ -7337,7 +7337,7 @@ Trace.endSection();
 			WindowManager windowManager = getWindowManager();
 			Display display = windowManager.getDefaultDisplay();
 			Display.Mode currentmode = display.getMode();
-			return (int)currentmode.getRefreshRate();
+			return java.lang.Math.round(currentmode.getRefreshRate());
 		}
 
 		return 60;
