Ensure that keyboard dismissed isn't called when text input is accepted by the Virtual Keyboard's done key or a physical keyboard's enter key.
Additionally, ensure that on dismissal of the virtual keyboard, text entry type is Canceled so that text committal only triggers if the widget is configured with behaviour 'Text commit on Dismiss'

diff --git a/Engine/Build/Android/Java/src/com/epicgames/unreal/GameActivity.java.template b/Engine/Build/Android/Java/src/com/epicgames/unreal/GameActivity.java.template
--- a/Engine/Build/Android/Java/src/com/epicgames/unreal/GameActivity.java.template
+++ b/Engine/Build/Android/Java/src/com/epicgames/unreal/GameActivity.java.template
@@ -454,6 +454,7 @@ public class GameActivity extends $${gameActivitySuperClass}$$ implements Surfac
 	// Keep a reference to the main content view so we can bring up the virtual keyboard without an editbox
 	private View mainView;
 	private boolean bKeyboardShowing;
+	private boolean bKeyboardInputCommittedByEnterKey = false;
 
 	private View mainDecorView;
 	private Rect mainDecorViewRect;
@@ -5836,6 +5837,13 @@ Trace.endSection();
 		}, lastVirtualKeyboardCommandDelay);
 	}
 	
+	private void CommitVirtualKeyboardInputWithEnterKey()
+	{
+		bKeyboardInputCommittedByEnterKey = true;
+		AndroidThunkJava_HideVirtualKeyboardInput();
+		nativeVirtualKeyboardSendKey(KeyEvent.KEYCODE_ENTER);
+	}
+
 	//jira UE-49141 Virtual keyboard is unresponsive with repeated tapping in control (some devices)
 	//#jira UE-49139 Tapping in the same text box doesn't make the virtual keyboard disappear
 	public synchronized void processLastVirtualKeyboardCommand()
@@ -5925,6 +5933,7 @@ Trace.endSection();
 
 					nativeVirtualKeyboardVisible(true);
 					bKeyboardShowing = true;
+					bKeyboardInputCommittedByEnterKey = false;
 
 					CheckKeyboardDisplayed();
 				}
@@ -5955,7 +5964,8 @@ Trace.endSection();
 					imm.hideSoftInputFromWindow(newVirtualKeyboardInput.getWindowToken(), 0);
 
 					String finalMessage = newVirtualKeyboardInput.getText().toString();
-					nativeVirtualKeyboardResult(true, finalMessage);
+					nativeVirtualKeyboardDismissed(!bKeyboardInputCommittedByEnterKey, finalMessage);
+					bKeyboardInputCommittedByEnterKey = false;
 					
 					nativeVirtualKeyboardVisible(false);
 					bKeyboardShowing = false;
@@ -8070,8 +8080,7 @@ Trace.endSection();
 						}
 						else
 						{
-							AndroidThunkJava_HideVirtualKeyboardInput();
-							nativeVirtualKeyboardSendKey(KeyEvent.KEYCODE_ENTER);
+							CommitVirtualKeyboardInputWithEnterKey();
 						}
 					}
 				}
@@ -8136,8 +8145,7 @@ Trace.endSection();
 						virtualKeyboardTextChangeRunnable.run();
 						virtualKeyboardTextChangeRunnable = null;
 					}
-					AndroidThunkJava_HideVirtualKeyboardInput();
-					nativeVirtualKeyboardSendKey(KeyEvent.KEYCODE_ENTER);
+					CommitVirtualKeyboardInputWithEnterKey();
 	                return true;
 				}
 	            return false;
@@ -8684,6 +8692,7 @@ Trace.endSection();
 	public native void nativeConsoleCommand(String commandString);
 	public native void nativeVirtualKeyboardChanged(String contents);
 	public native void nativeVirtualKeyboardResult(boolean update, String contents);
+	public native void nativeVirtualKeyboardDismissed(boolean update, String contents);
 	public native void nativeVirtualKeyboardSendKey(int keyCode);
 	public native void nativeVirtualKeyboardSendTextSelection(String contents, int selStart, int selEnd);
 	public native void nativeVirtualKeyboardSendSelection(int selStart, int selEnd);
diff --git a/Engine/Source/Runtime/Core/Public/Android/AndroidJniGameActivity.h b/Engine/Source/Runtime/Core/Public/Android/AndroidJniGameActivity.h
--- a/Engine/Source/Runtime/Core/Public/Android/AndroidJniGameActivity.h
+++ b/Engine/Source/Runtime/Core/Public/Android/AndroidJniGameActivity.h
@@ -41,6 +41,7 @@ namespace UE::Jni
 		static void JNICALL nativeVirtualKeyboardVisible(JNIEnv* env, jobject thiz, jboolean bShown);
 		static void JNICALL nativeVirtualKeyboardShown(JNIEnv* env, jobject thiz, jint left, jint top, jint right, jint bottom);
 		static void JNICALL nativeVirtualKeyboardResult(JNIEnv* env, jobject thiz, jboolean update, jstring contents);
+		static void JNICALL nativeVirtualKeyboardDismissed(JNIEnv* env, jobject thiz, jboolean update, jstring contents);
 		static void JNICALL nativeVirtualKeyboardChanged(JNIEnv* env, jobject thiz, jstring contents);
 		static void JNICALL nativeVirtualKeyboardSendKey(JNIEnv* env, jobject thiz, jint keyCode);
 		static void JNICALL nativeVirtualKeyboardSendSelection(JNIEnv* jenv, jobject thiz, jint selStart, jint selEnd);
@@ -88,6 +89,7 @@ namespace UE::Jni
 			UE_JNI_NATIVE_METHOD(nativeVirtualKeyboardVisible),
 			UE_JNI_NATIVE_METHOD(nativeVirtualKeyboardShown),
 			UE_JNI_NATIVE_METHOD(nativeVirtualKeyboardResult),
+			UE_JNI_NATIVE_METHOD(nativeVirtualKeyboardDismissed),
 			UE_JNI_NATIVE_METHOD(nativeVirtualKeyboardChanged),
 			UE_JNI_NATIVE_METHOD(nativeVirtualKeyboardSendKey),
 			UE_JNI_NATIVE_METHOD(nativeVirtualKeyboardSendSelection),
diff --git a/Engine/Source/Runtime/Launch/Private/Android/AndroidJNI.cpp b/Engine/Source/Runtime/Launch/Private/Android/AndroidJNI.cpp
--- a/Engine/Source/Runtime/Launch/Private/Android/AndroidJNI.cpp
+++ b/Engine/Source/Runtime/Launch/Private/Android/AndroidJNI.cpp
@@ -1173,7 +1173,7 @@ void AndroidThunkCpp_ShowVirtualKeyboardInput(TSharedPtr<IVirtualKeyboardEntry>
 }
 
 //This function is declared in the Java-defined class, GameActivity.java: "public native void nativeVirtualKeyboardResult(bool update, String contents);"
-void JNICALL UE::Jni::FGameActivity::nativeVirtualKeyboardResult(JNIEnv* jenv, jobject thiz, jboolean update, jstring contents)
+static void nativeVirtualKeyboardResult(JNIEnv* jenv, jobject thiz, jboolean update, jstring contents, ETextEntryType textEntryType)
 {
 	// update text widget with new contents if OK pressed
 	if (update == JNI_TRUE)
@@ -1185,12 +1185,12 @@ void JNICALL UE::Jni::FGameActivity::nativeVirtualKeyboardResult(JNIEnv* jenv, j
 			// call to set the widget text on game thread
 			if (FTaskGraphInterface::IsRunning())
 			{
-				FFunctionGraphTask::CreateAndDispatchWhenReady([Contents=Contents]()
+				FFunctionGraphTask::CreateAndDispatchWhenReady([Contents=Contents, textEntryType=textEntryType]()
 				{
 					TSharedPtr<IVirtualKeyboardEntry> LockedKeyboardWidget(VirtualKeyboardWidget.Pin());
 					if (LockedKeyboardWidget.IsValid())
 					{
-						LockedKeyboardWidget->SetTextFromVirtualKeyboard(FText::FromString(Contents), ETextEntryType::TextEntryAccepted);
+						LockedKeyboardWidget->SetTextFromVirtualKeyboard(FText::FromString(Contents), textEntryType);
 					}
 
 					// release reference
@@ -1211,6 +1211,16 @@ void JNICALL UE::Jni::FGameActivity::nativeVirtualKeyboardResult(JNIEnv* jenv, j
 	}
 }
 
+//This function is declared in the Java-defined class, GameActivity.java: "public native void nativeVirtualKeyboardResult(bool update, String contents);"
+void JNICALL UE::Jni::FGameActivity::nativeVirtualKeyboardResult(JNIEnv* jenv, jobject thiz, jboolean update, jstring contents)
+{
+	::nativeVirtualKeyboardResult(jenv, thiz, update, contents, ETextEntryType::TextEntryAccepted);
+}
+void JNICALL UE::Jni::FGameActivity::nativeVirtualKeyboardDismissed(JNIEnv* jenv, jobject thiz, jboolean update, jstring contents)
+{
+	::nativeVirtualKeyboardResult(jenv, thiz, update, contents, ETextEntryType::TextEntryCanceled);
+}
+
 //This function is declared in the Java-defined class, GameActivity.java: "public native void nativeVirtualKeyboardChanged(String contents);"
 void JNICALL UE::Jni::FGameActivity::nativeVirtualKeyboardChanged(JNIEnv* jenv, jobject thiz, jstring contents)
 {
