Fixing white wash effect applied to game when using a controller in OS modal on Android
- GameActivity.java.template: Removing surface view's TRANSPARENT background color

#rnx

diff --git a/Engine/Build/Android/Java/src/com/epicgames/unreal/GameActivity.java.template b/Engine/Build/Android/Java/src/com/epicgames/unreal/GameActivity.java.template
--- a/Engine/Build/Android/Java/src/com/epicgames/unreal/GameActivity.java.template
+++ b/Engine/Build/Android/Java/src/com/epicgames/unreal/GameActivity.java.template
@@ -4700,7 +4700,6 @@ Trace.beginSection("UI Building");
 			// Need to create our surface view here regardless of if we are going to end up using it
 			getWindow().takeSurface(null);
 			MySurfaceView = new SurfaceView(getThisActivity());
-			MySurfaceView.setBackgroundColor(Color.TRANSPARENT);
 			MySurfaceView.getHolder().addCallback(this);
 			containerFrameLayout.addView(MySurfaceView);
 			containerFrameLayout.addView(virtualKeyboardLayout);
