Add CRASH_CONTEXT_SCOPE to vulkan graphics PSO create, provides hash information in the event of a crash during compile. disabled by default.
Add ability to include files in FCrashContextExtendedWriter.
Android:
Add an option to capture vulkan serialized PSO create data via the external PSO services. disabled by default.
when enabled also attach PSO create data in crash reports.

diff --git a/Engine/Build/Android/Java/src/com/epicgames/unreal/psoservices/PSOProgramServiceAccessor.java b/Engine/Build/Android/Java/src/com/epicgames/unreal/psoservices/PSOProgramServiceAccessor.java
--- a/Engine/Build/Android/Java/src/com/epicgames/unreal/psoservices/PSOProgramServiceAccessor.java
+++ b/Engine/Build/Android/Java/src/com/epicgames/unreal/psoservices/PSOProgramServiceAccessor.java
@@ -6,6 +6,7 @@ import android.content.ComponentName;
 import android.content.Context;
 import android.content.Intent;
 import android.content.ServiceConnection;
+import android.os.Build;
 import android.os.Bundle;
 import android.os.Handler;
 import android.os.HandlerThread;
@@ -14,8 +15,10 @@ import android.os.Looper;
 import android.os.Message;
 import android.os.Messenger;
 import android.os.RemoteException;
+import android.os.SharedMemory;
 import android.os.StrictMode;
 import android.os.Trace;
+import android.system.ErrnoException;
 import android.util.Log;
 import android.os.ParcelFileDescriptor;
 import android.os.SystemClock;
@@ -25,10 +28,16 @@ import com.epicgames.unreal.Logger;
 import com.google.android.gms.games.Game;
 
 import java.io.BufferedReader;
+import java.io.DataOutputStream;
 import java.io.File;
 import java.io.FileNotFoundException;
+import java.io.FileOutputStream;
 import java.io.FileReader;
 import java.io.IOException;
+import java.nio.ByteBuffer;
+import java.nio.channels.FileChannel;
+import java.nio.file.Path;
+import java.nio.file.Paths;
 import java.util.HashMap;
 import java.util.HashSet;
 import java.util.Map;
@@ -40,6 +49,7 @@ import java.util.concurrent.atomic.AtomicBoolean;
 import java.util.concurrent.locks.ReadWriteLock;
 import java.util.concurrent.locks.ReentrantReadWriteLock;
 import java.util.concurrent.TimeUnit;
+import java.nio.file.Files;
 
 public class PSOProgramServiceAccessor
 {
@@ -179,6 +189,10 @@ public class PSOProgramServiceAccessor
 	private boolean bServicePriHigh = false;
 	public void SetServicePriority(boolean bHigh) { bServicePriHigh = bHigh; }
 
+	// When  >0 the PSO init data sent to each service is retained. a total of RetainedProgramInitDataHistoryLimit * servicecount files will be stored. 
+	private static int RetainedProgramInitDataHistoryLimit = 0;
+
+	private static String RetainedProgramInitDataHistoryPath = "/PSORetained/";
 	public PSOProgramServiceAccessor()
 	{
 		//StrictMode.enableDefaults();
@@ -197,24 +211,38 @@ public class PSOProgramServiceAccessor
 		_PSOProgramServiceAccessor = this;
 	}
 
-	private void ClearServiceLogDirs()
+	// deletes a directory and its content. 
+	// Note: Doesn't delete subdirs.
+	private static void DeleteDir(String DirToRemove)
 	{
-		String OGLLogDir = mActivity.getFilesDir()+PSOProgramService.LogDir;
-		File deleteDir = new File(OGLLogDir);
-		boolean bDeleteDirExists = deleteDir.exists();
-		if( bDeleteDirExists )
+		try
 		{
-			File[] deletedirContent = deleteDir.listFiles();
-			if(deletedirContent != null)
+			File deleteDir = new File(DirToRemove);
+			boolean bDeleteDirExists = deleteDir.exists();
+			if (bDeleteDirExists)
 			{
-				for(File deleteFile : deletedirContent)
+				File[] deletedirContent = deleteDir.listFiles();
+				if (deletedirContent != null)
 				{
-					deleteFile.delete();
+					for (File deleteFile : deletedirContent)
+					{
+						deleteFile.delete();
+					}
 				}
+				deleteDir.delete();
 			}
-			deleteDir.delete();
+		}
+		catch(Exception e)
+		{
+			Log.error( "DeleteDir("+DirToRemove+") - FAIL " + e);
+			e.printStackTrace();
 		}
 	}
+	private void ClearServiceLogDirs()
+	{
+		String OGLLogDir = mActivity.getFilesDir()+PSOProgramService.LogDir;
+		DeleteDir(OGLLogDir);
+	}
 
 	public static boolean AndroidThunkJava_StartRemoteProgramLink(int numServices, boolean bUseRobustEGLContext, boolean bUseVulkan)
 	{
@@ -223,6 +251,13 @@ public class PSOProgramServiceAccessor
 		{
 			try
 			{
+				// always clear the retained program dir
+				String RetainedProgramDestDir = _PSOProgramServiceAccessor.GetRetainedProgramInitDataPath();
+				DeleteDir(RetainedProgramDestDir);
+				if(RetainedProgramInitDataHistoryLimit > 0)
+				{
+					Files.createDirectories( Paths.get( RetainedProgramDestDir ) );
+				}
 				ProgramServiceAccessorlock.writeLock().lock();
 				PSOProgramAccessorHandlerThread = new HandlerThread("PSOProgramAccessorThread");
 				PSOProgramAccessorHandlerThread.start();
@@ -231,7 +266,7 @@ public class PSOProgramServiceAccessor
 				// rate limit the total number of inflight PSO compile requests. shut down waits for all jobs to complete and this reduces the potential wait time.
 				// calling threads wait for a request permit to become available or shutdown is issued.
 				MaxLiveRequestPermitCount = numServices * 3;
-				MaxLiveRequestsSemaphore = new Semaphore(MaxLiveRequestPermitCount); 
+				MaxLiveRequestsSemaphore = new Semaphore(MaxLiveRequestPermitCount);
 				if(bUseVulkan)
 				{
 					bSuccess = _PSOProgramServiceAccessor.StartVulkanServices(numServices);
@@ -340,6 +375,57 @@ public class PSOProgramServiceAccessor
 		return bSuccess;
 	}
 
+	public String GetRetainedProgramInitDataPath()
+	{
+		return mActivity.getFilesDir() + RetainedProgramInitDataHistoryPath;
+	}
+	
+	public static void AndroidThunkJava_SetProgramInitDataRetainLimit(int InRetainedProgramInitDataHistoryLimit)
+	{
+		if (Build.VERSION.SDK_INT >= Build.VERSION_CODES.O_MR1)
+		{
+			RetainedProgramInitDataHistoryLimit = InRetainedProgramInitDataHistoryLimit;
+			try
+			{
+				ProgramServiceAccessorlock.readLock().lock();
+				if( _PSOProgramServiceAccessor != null)
+				{
+					String DestDir = _PSOProgramServiceAccessor.GetRetainedProgramInitDataPath();
+					Files.createDirectories( Paths.get( DestDir ) );
+				}
+			}
+			catch (Exception e)
+			{
+				Log.error( "SetProgramInitDataRetainLimit FAIL " + e);
+			}
+		    finally
+		    {
+			    ProgramServiceAccessorlock.readLock().unlock();
+		    }
+		}
+		else
+		{
+			Log.error( "SetProgramInitDataRetainLimit requires API 27 and above.");
+		}
+	}
+
+	public static String AndroidThunkJava_GetRetainedProgramInitDataPath()
+	{
+		try
+		{
+			ProgramServiceAccessorlock.readLock().lock();
+			if( _PSOProgramServiceAccessor != null)
+			{
+				return _PSOProgramServiceAccessor.GetRetainedProgramInitDataPath();
+			}
+		}
+		finally
+		{
+			ProgramServiceAccessorlock.readLock().unlock();
+		}
+		return "";
+	}
+
 	void StopServices()
 	{
 		if(ServiceInstances != null)
@@ -856,6 +942,23 @@ public class PSOProgramServiceAccessor
 					}
 				}
 
+				if( RetainedProgramInitDataHistoryLimit > 0 )
+				{
+					// Attempt to capture the input params used to create the PSO
+					// this can be useful to gather data required to repro PSO compiler crashes.
+					if( thisInstance  != null )
+					{
+						// Generate the path for the historical creation data. 
+						// the filename is in the form: pso_{service_index}_{service_local_job_idx % RetainedProgramInitDataHistoryLimit}
+						int ServiceIdx = params.getInt(PSOProgramService.ServiceID_Key);
+						String filename = "pso_" +ServiceIdx+"_"+(thisInstance.NumJobsSent % RetainedProgramInitDataHistoryLimit);
+						Path FullPath = Paths.get(AndroidThunkJava_GetRetainedProgramInitDataPath(), filename);
+						
+						CopyFromFileDescriptor(thisInstance.NumJobsSent, ContextData, VertexShaderSize, PixelShaderSize, PSODataSize, PSOCacheDataSize, parcelFD, FullPath);
+						thisInstance.NumJobsSent++;
+					}
+				}
+
 				//Log.verbose("OGLRemoteProgramLink dispatching to "+thisInstance.Name()+" (" + ThisJobID + ")");
 				pendingResponse.ResponseState = JobResponse.ResponseStateEnum.Pending;
 				if( thisInstance != null && thisInstance.SendMessage(msg) )
@@ -938,10 +1041,11 @@ public class PSOProgramServiceAccessor
 		catch (Exception e)
 		{
 			Log.error( "AndroidThunkJava_VKPSOGFXCompileShm FAIL ("+JobID+")" + e);
+			e.printStackTrace();
 		}
 		finally
 		{
-			if(parcelFD != null) 
+			if(parcelFD != null)
 			{
 				try
 				{
@@ -950,6 +1054,7 @@ public class PSOProgramServiceAccessor
 				catch (IOException e)
 				{
 					Log.error( "AndroidThunkJava_VKPSOGFXCompileShm FAIL ("+JobID+")" + e);
+					e.printStackTrace();
 				}
 			}
 			ProgramServiceAccessorlock.readLock().unlock();
@@ -959,10 +1064,43 @@ public class PSOProgramServiceAccessor
 		return null;
 	}
 
+	private static void CopyFromFileDescriptor(long LocalJobID, byte[] ContextData, long VertexShaderSize, long PixelShaderSize, long PSODataSize, long PSOCacheDataSize, ParcelFileDescriptor fd, Path FullPath)
+	{
+		if( Build.VERSION.SDK_INT < Build.VERSION_CODES.O_MR1)
+		{
+			return;
+		}
+
+		try (SharedMemory shm = SharedMemory.fromFileDescriptor(fd.dup()); FileOutputStream OutStream = new FileOutputStream(String.valueOf(FullPath)); DataOutputStream DataOut = new DataOutputStream(OutStream); FileChannel OutChannel = OutStream.getChannel();)
+		{
+			DataOut.writeLong(LocalJobID);
+			DataOut.writeLong(VertexShaderSize);
+			DataOut.writeLong(PixelShaderSize);
+			DataOut.writeLong(PSODataSize);
+			DataOut.writeLong(PSOCacheDataSize);
+			DataOut.writeLong(ContextData.length);
+			DataOut.write(ContextData);
+			DataOut.flush();
+
+			ByteBuffer DataBuffer = shm.mapReadOnly();
+			DataBuffer.position(0);
+			while(DataBuffer.hasRemaining())
+			{
+				OutChannel.write(DataBuffer);
+			}
+			SharedMemory.unmap(DataBuffer);
+		}
+		catch (Exception e)
+		{
+			Log.error("CopyFromFileDescriptor FAILED" + e + " : " + e.getStackTrace());
+		}
+	}
+
 	class OGLServiceInstance
 	{
 		private final Class ServiceClass;
 		private final boolean bRobustContext;
+		private int NumJobsSent = 0;
 		boolean mShouldUnbind = false;
 		private final AtomicInteger mBound = new AtomicInteger(-1);
 		Messenger mService = null;
@@ -1091,7 +1229,7 @@ public class PSOProgramServiceAccessor
 				{
 					BindServiceFlags = BindServiceFlags | Context.BIND_IMPORTANT;
 				}
-				 
+
 				intent.putExtra(PSOProgramService.RobustContextKey,bRobustContext);
 
 				mShouldUnbind = mContext.bindService(intent, mConnection, BindServiceFlags);
diff --git a/Engine/Source/Runtime/Core/Private/Android/AndroidPlatformFile.cpp b/Engine/Source/Runtime/Core/Private/Android/AndroidPlatformFile.cpp
--- a/Engine/Source/Runtime/Core/Private/Android/AndroidPlatformFile.cpp
+++ b/Engine/Source/Runtime/Core/Private/Android/AndroidPlatformFile.cpp
@@ -1663,7 +1663,14 @@ public:
 
 		if (IsLocal(FromLocalPath))
 		{
-			return rename(TCHAR_TO_UTF8(*FromLocalPath), TCHAR_TO_UTF8(*ToLocalPath)) != -1;
+			bool bSuccess = rename(TCHAR_TO_UTF8(*FromLocalPath), TCHAR_TO_UTF8(*ToLocalPath)) != -1;
+			if (!bSuccess)
+			{
+#if LOG_ANDROID_FILE
+				FPlatformMisc::LowLevelOutputDebugStringf(TEXT("FAndroidPlatformFile::MoveFile() - failed = %s"), UTF8_TO_TCHAR(strerror(errno)));
+#endif
+			}
+			return bSuccess;
 		}
 		return false;
 	}
diff --git a/Engine/Source/Runtime/Core/Private/GenericPlatform/GenericPlatformCrashContext.cpp b/Engine/Source/Runtime/Core/Private/GenericPlatform/GenericPlatformCrashContext.cpp
--- a/Engine/Source/Runtime/Core/Private/GenericPlatform/GenericPlatformCrashContext.cpp
+++ b/Engine/Source/Runtime/Core/Private/GenericPlatform/GenericPlatformCrashContext.cpp
@@ -1897,6 +1897,33 @@ struct FCrashContextExtendedWriterImpl : public FCrashContextExtendedWriter
 		OutputBuffer(Identifier, (uint8*)Converter.Get(), Converter.Length(), TEXT("txt"));
 	}
 
+	/** MoveFile a file into the report */
+	virtual bool MoveFile(const TCHAR* FromFilePath) override
+	{
+		TCHAR DestFilePath[1024] = { 0 };
+		bool bFileMoved = false;
+		FString ToFileName = *FPaths::GetCleanFilename(FromFilePath);
+		int32 DestPathLength = FCString::Snprintf(DestFilePath, UE_ARRAY_COUNT(DestFilePath), TEXT("%s/%s"), OutputDirectory, *ToFileName);
+		if (DestPathLength>0 && DestPathLength < UE_ARRAY_COUNT(DestFilePath))
+		{
+			IPlatformFile& PlatformFile = IPlatformFile::GetPlatformPhysical();
+			bFileMoved = PlatformFile.MoveFile(DestFilePath, FromFilePath);
+			if (!bFileMoved) // attempt copy the file if the move fails. (cant move across file systems etc)
+			{
+				bFileMoved = PlatformFile.CopyFile(DestFilePath, FromFilePath);
+				if (bFileMoved)
+				{
+					PlatformFile.DeleteFile(FromFilePath);
+				}
+				else
+				{
+					FPlatformMisc::LowLevelOutputDebugStringf(TEXT("Crash Context FileMove failed. (%s , %s)"), FromFilePath, DestFilePath);
+				}
+			}
+		}
+		return bFileMoved;
+	}
+
 private:
 	const TCHAR* OutputDirectory;
 };
diff --git a/Engine/Source/Runtime/Core/Public/GenericPlatform/GenericPlatformCrashContext.h b/Engine/Source/Runtime/Core/Public/GenericPlatform/GenericPlatformCrashContext.h
--- a/Engine/Source/Runtime/Core/Public/GenericPlatform/GenericPlatformCrashContext.h
+++ b/Engine/Source/Runtime/Core/Public/GenericPlatform/GenericPlatformCrashContext.h
@@ -305,6 +305,12 @@ struct FCrashContextExtendedWriter
 
 	/** Add a named buffer containing a string to the report. */
 	virtual void AddString(const TCHAR* Identifier, const TCHAR* DataStr) = 0;
+
+	/** MoveFile a file into the report
+	* FromFilePath - the full path of the file you wish to move into the report.
+	* The base filename is used for the destination file.
+	*/
+	virtual bool MoveFile(const TCHAR* FromFilePath) = 0;
 };
 
 /** Simple Delegate for additional crash context. */
diff --git a/Engine/Source/Runtime/VulkanRHI/Private/Android/VulkanAndroidPlatform.cpp b/Engine/Source/Runtime/VulkanRHI/Private/Android/VulkanAndroidPlatform.cpp
--- a/Engine/Source/Runtime/VulkanRHI/Private/Android/VulkanAndroidPlatform.cpp
+++ b/Engine/Source/Runtime/VulkanRHI/Private/Android/VulkanAndroidPlatform.cpp
@@ -22,6 +22,7 @@
 #include "Android/AndroidDynamicRHI.h"
 #include "PSOMetrics.h"
 #include "../VulkanContext.h"
+#include "GenericPlatform/GenericPlatformCrashContext.h"
 
 #if USE_ANDROID_VULKAN_SWAPPY
 #undef VK_NO_PROTOTYPES
@@ -1627,6 +1628,8 @@ struct FVKRemoteProgramCompileJNI
 	jmethodID HaveServicesFailed = 0;
 	jmethodID AreProgramServicesReady = 0;	
 	jmethodID StopRemoteProgramLink = 0;
+	jmethodID SetProgramInitDataRetainLimit = 0;
+	jmethodID GetRetainedProgramInitDataPath = 0;
 	jclass ProgramResponseClass = 0;
 	jfieldID ProgramResponse_SuccessField = 0;
 	jfieldID ProgramResponse_ErrorField = 0;
@@ -1667,6 +1670,10 @@ struct FVKRemoteProgramCompileJNI
 			CHECK_JNI_EXCEPTIONS(Env);
 			StopRemoteProgramLink = FJavaWrapper::FindStaticMethod(Env, PSOServiceAccessor, "AndroidThunkJava_StopRemoteProgramLink", "()V", false);
 			CHECK_JNI_EXCEPTIONS(Env);
+			SetProgramInitDataRetainLimit = FJavaWrapper::FindStaticMethod(Env, PSOServiceAccessor, "AndroidThunkJava_SetProgramInitDataRetainLimit", "(I)V", false);
+			CHECK_JNI_EXCEPTIONS(Env);
+			GetRetainedProgramInitDataPath = FJavaWrapper::FindStaticMethod(Env, PSOServiceAccessor, "AndroidThunkJava_GetRetainedProgramInitDataPath", "()Ljava/lang/String;", false);
+			CHECK_JNI_EXCEPTIONS(Env);
 			ProgramResponseClass = AndroidJavaEnv::FindJavaClassGlobalRef("com/epicgames/unreal/psoservices/PSOProgramServiceAccessor$JNIProgramLinkResponse");
 			CHECK_JNI_EXCEPTIONS(Env);
 			ProgramResponse_SuccessField = FJavaWrapper::FindField(Env, ProgramResponseClass, "bCompileSuccess", "Z", true);
@@ -1681,7 +1688,8 @@ struct FVKRemoteProgramCompileJNI
 			CHECK_JNI_EXCEPTIONS(Env);
 		}
 
-		bAllFound = PSOServiceAccessor && DispatchPSOCompile && DispatchPSOCompileShm && StartRemoteProgramLink && HaveServicesFailed && AreProgramServicesReady && StopRemoteProgramLink && ProgramResponseClass && ProgramResponse_SuccessField && ProgramResponse_CompiledBinaryField && ProgramResponse_ErrorField && ProgramResponse_SHMOutputHandleField && ProgramResponse_CompilationDurationField;
+		bAllFound = PSOServiceAccessor && DispatchPSOCompile && DispatchPSOCompileShm && StartRemoteProgramLink && HaveServicesFailed && AreProgramServicesReady && StopRemoteProgramLink && SetProgramInitDataRetainLimit && GetRetainedProgramInitDataPath
+			&& ProgramResponseClass && ProgramResponse_SuccessField && ProgramResponse_CompiledBinaryField && ProgramResponse_ErrorField && ProgramResponse_SHMOutputHandleField && ProgramResponse_CompilationDurationField;
 		UE_CLOG(!bAllFound, LogRHI, Fatal, TEXT("Failed to find JNI Vulkan remote program compiler."));
 	}
 }VKRemoteProgramCompileJNI;
@@ -1733,6 +1741,76 @@ bool FVulkanAndroidPlatform::AreRemoteCompileServicesActive()
 	return false;
 }
 
+
+namespace AndroidVulkan
+{
+	namespace PSOCapture
+	{
+#if WITH_ADDITIONAL_CRASH_CONTEXTS
+		int32 GPSOInitHistoryLimit = 0;
+		static FAutoConsoleVariableRef CVarVulkanExtensionFramePacer(
+			TEXT("android.Vulkan.RetainLimitForPSOCreateCrashReport"),
+			GPSOInitHistoryLimit,
+			TEXT("Retain the PSO creation source data from the external compiler services when > 0. If a crash occurs the data is added to the crash report.\n")
+			TEXT(" 0 to disable (default)\n")
+			TEXT(" >0 == amount of history to retain per service."),
+			ECVF_RenderThreadSafe
+		);
+
+		static FString PSOCrashHistoryDir;
+
+		static void OnCrashDumpPSOHistoryInfo(FCrashContextExtendedWriter& Writer)
+		{
+			if (GPSOInitHistoryLimit > 0 && FVulkanAndroidPlatform::AreRemoteCompileServicesActive())
+			{
+				FPlatformMisc::LowLevelOutputDebugStringf(TEXT("OnCrashDumpPSOHistoryInfo"));
+				IPlatformFile::GetPlatformPhysical().IterateDirectory(*PSOCrashHistoryDir, [&Writer](const TCHAR* Pathname, bool bIsDirectory)
+					{
+						if(!bIsDirectory)
+						{
+							FPlatformMisc::LowLevelOutputDebugStringf(TEXT("OnCrashDumpPSOHistoryInfo: %s -> %s"), Pathname, *FPaths::GetCleanFilename(Pathname));
+							Writer.MoveFile(Pathname);
+						}
+						return true;
+					});
+			}
+		}
+#endif
+		static void InitCrashDataCapture()
+		{
+#if WITH_ADDITIONAL_CRASH_CONTEXTS
+			UE_LOG(LogRHI, Log, TEXT("Retaining remote compile services history, limit = %d"), AndroidVulkan::PSOCapture::GPSOInitHistoryLimit);
+
+			if (AndroidVulkan::PSOCapture::GPSOInitHistoryLimit > 0)
+			{
+				check(IsInRenderingThread());
+				check(PSOCrashHistoryDir.IsEmpty());
+				JNIEnv* Env = FAndroidApplication::GetJavaEnv();
+				Env->CallStaticVoidMethod(VKRemoteProgramCompileJNI.PSOServiceAccessor, VKRemoteProgramCompileJNI.SetProgramInitDataRetainLimit, (jint)AndroidVulkan::PSOCapture::GPSOInitHistoryLimit);
+				PSOCrashHistoryDir = FJavaHelper::FStringFromLocalRef(Env, (jstring)Env->CallStaticObjectMethod(VKRemoteProgramCompileJNI.PSOServiceAccessor, VKRemoteProgramCompileJNI.GetRetainedProgramInitDataPath));
+				CHECK_JNI_EXCEPTIONS(Env);
+			}
+#endif
+		}
+	}
+}                                                                                                                                                                                                                                                                                                                                              
+
+void FVulkanAndroidPlatform::OnCreateGfxPipelineCrash(FCrashContextExtendedWriter& Writer, bool bIsPrecache)
+{
+#if WITH_ADDITIONAL_CRASH_CONTEXTS
+	if(bIsPrecache)
+	{
+		// we've crashed during PSO compile of a precache PSO, dump all the retained history.
+		// just do this once.
+		static std::atomic_bool bDoneOnce = false;
+		if (bDoneOnce.exchange(true) == false)
+		{
+			AndroidVulkan::PSOCapture::OnCrashDumpPSOHistoryInfo(Writer);
+		}
+	}
+#endif //WITH_ADDITIONAL_CRASH_CONTEXTS
+}
+
 bool FVulkanAndroidPlatform::StartRemoteCompileServices(int NumServices)
 {
 	JNIEnv* Env = FAndroidApplication::GetJavaEnv();
@@ -1742,6 +1820,8 @@ bool FVulkanAndroidPlatform::StartRemoteCompileServices(int NumServices)
 	if (Env && AreAndroidVulkanRemoteCompileServicesAvailable() && !AndroidVulkanService::GRemoteCompileServicesStarted)
 	{
 		AndroidVulkanService::GRemoteCompileServicesStarted = (bool)Env->CallStaticBooleanMethod(VKRemoteProgramCompileJNI.PSOServiceAccessor, VKRemoteProgramCompileJNI.StartRemoteProgramLink, (jint)NumServices, /*bUseRobustEGLContext*/(jboolean)false, /*bUseVulkan*/(jboolean)true);
+
+		AndroidVulkan::PSOCapture::InitCrashDataCapture();
 	}
 
 	return AndroidVulkanService::GRemoteCompileServicesStarted;
@@ -1774,14 +1854,15 @@ FString FVulkanAndroidPlatform::GetCompiledPSOCacheTopFolderPath()
 }
 
 VkPipelineCache FVulkanAndroidPlatform::PrecompilePSO(
-		FVulkanDevice* Device, 
-		const TArrayView<uint8> OptionalPSOCacheData, 
+		FVulkanDevice* Device,
+		const TArrayView<const uint8> ContextData,
+		const TArrayView<const uint8> OptionalPSOCacheData, 
 		FGraphicsPipelineStateInitializer::EPSOPrecacheCompileType PSOCompileType,
 		const VkGraphicsPipelineCreateInfo* PipelineInfo, 
 		const FGfxPipelineDesc* GfxEntry, 
 		const FVulkanRenderTargetLayout* RTLayout, 
-		TArrayView<uint32_t> VS, 
-		TArrayView<uint32_t> PS, 
+		TArrayView<const uint32_t> VS, 
+		TArrayView<const uint32_t> PS, 
 		size_t& AfterSize,
 		FString* FailureMessageOUT
 	)
@@ -1817,8 +1898,9 @@ VkPipelineCache FVulkanAndroidPlatform::PrecompilePSO(
 		const uint64 PreSuppliedCacheSize = ArrayByteSize(OptionalPSOCacheData);
 		const uint64 TotalOutputSize = VSSize + PSSize + PSOParamsSize + PreSuppliedCacheSize;
 
-		auto ProgramKeyBuffer = NewScopedJavaObject(Env, Env->NewByteArray(4));
-		Env->SetByteArrayRegion(*ProgramKeyBuffer, 0, 4, reinterpret_cast<const jbyte*>("Test"));
+		auto ContextJavaArray = NewScopedJavaObject(Env, Env->NewByteArray(ContextData.Num()));
+		Env->SetByteArrayRegion(*ContextJavaArray, 0, ContextData.Num(), reinterpret_cast<const jbyte*>(ContextData.GetData()));
+		CHECK_JNI_EXCEPTIONS(Env);
 
 		// create a shared mem region for external process to access.
 		FScopedJavaObject<_jobject*> ProgramResponseObj;
@@ -1858,7 +1940,7 @@ VkPipelineCache FVulkanAndroidPlatform::PrecompilePSO(
 					bool bEnableTimeOuts = !FPlatformMisc::IsDebuggerPresent();
 					{
 						QUICK_SCOPE_CYCLE_COUNTER(STAT_VulkanAndroid_PrecompilePSOJAVA);
-						ProgramResponseObj = NewScopedJavaObject(Env, Env->CallStaticObjectMethod(VKRemoteProgramCompileJNI.PSOServiceAccessor, VKRemoteProgramCompileJNI.DispatchPSOCompileShm, *ProgramKeyBuffer, PriorityInfo.GetPriorityInfo(), SharedMemFD, VSSize, PSSize, PSOParamsSize, PreSuppliedCacheSize, bEnableTimeOuts));
+						ProgramResponseObj = NewScopedJavaObject(Env, Env->CallStaticObjectMethod(VKRemoteProgramCompileJNI.PSOServiceAccessor, VKRemoteProgramCompileJNI.DispatchPSOCompileShm, *ContextJavaArray, PriorityInfo.GetPriorityInfo(), SharedMemFD, VSSize, PSSize, PSOParamsSize, PreSuppliedCacheSize, bEnableTimeOuts));
 					}
 					CHECK_JNI_EXCEPTIONS(Env);
 					munmap(SharedBuffer, memSize);
diff --git a/Engine/Source/Runtime/VulkanRHI/Private/Android/VulkanAndroidPlatform.h b/Engine/Source/Runtime/VulkanRHI/Private/Android/VulkanAndroidPlatform.h
--- a/Engine/Source/Runtime/VulkanRHI/Private/Android/VulkanAndroidPlatform.h
+++ b/Engine/Source/Runtime/VulkanRHI/Private/Android/VulkanAndroidPlatform.h
@@ -218,7 +218,7 @@ public:
 	static void DestroySwapchainKHR(VkDevice Device, VkSwapchainKHR Swapchain, const VkAllocationCallbacks* Allocator);
 
 	// handle precompile of PSOs, send to an android specific precompile external process.
-	static VkPipelineCache PrecompilePSO(FVulkanDevice* Device, const TArrayView<uint8> OptionalPSOCacheData, FGraphicsPipelineStateInitializer::EPSOPrecacheCompileType PSOCompileType, const VkGraphicsPipelineCreateInfo* PipelineInfo, const FGfxPipelineDesc* GfxEntry, const FVulkanRenderTargetLayout* RTLayout, TArrayView<uint32_t> VS, TArrayView<uint32_t> PS, size_t& AfterSize, FString* FailureMessageOUT = nullptr);
+	static VkPipelineCache PrecompilePSO(FVulkanDevice* Device, const TArrayView<const uint8> ContextData, const TArrayView<const uint8> OptionalPSOCacheData, FGraphicsPipelineStateInitializer::EPSOPrecacheCompileType PSOCompileType, const VkGraphicsPipelineCreateInfo* PipelineInfo, const FGfxPipelineDesc* GfxEntry, const FVulkanRenderTargetLayout* RTLayout, TArrayView<const uint32_t> VS, TArrayView<const uint32_t> PS, size_t& AfterSize, FString* FailureMessageOUT = nullptr);
 
 	static FString GetCompiledPSOCacheTopFolderPath();
 
@@ -258,6 +258,8 @@ public:
 
 	static VkTimeDomainKHR GetTimeDomain() { return VK_TIME_DOMAIN_CLOCK_MONOTONIC_KHR; }
 
+	static void OnCreateGfxPipelineCrash(struct FCrashContextExtendedWriter& Writer, bool bIsPrecache);
+
 #if USE_ANDROID_VULKAN_SWAPPY
 	static bool bSwappyEnabledAtRHIInit;
 #endif
diff --git a/Engine/Source/Runtime/VulkanRHI/Private/VulkanGenericPlatform.h b/Engine/Source/Runtime/VulkanRHI/Private/VulkanGenericPlatform.h
--- a/Engine/Source/Runtime/VulkanRHI/Private/VulkanGenericPlatform.h
+++ b/Engine/Source/Runtime/VulkanRHI/Private/VulkanGenericPlatform.h
@@ -5,6 +5,7 @@
 #include "Containers/Array.h"
 #include "Containers/ArrayView.h"
 #include "PixelFormat.h"
+#include "RHIResources.h"
 #include "Templates/UniquePtr.h"
 #include "VulkanThirdParty.h"
 
@@ -143,8 +144,8 @@ public:
 	// Gathers a list of pso cache filenames to attempt to load
 	static TArray<FString> GetPSOCacheFilenames();
 
-	// Gives platform a chance to handle precompile of PSOs, returns nullptr if unsupported
-	static VkPipelineCache PrecompilePSO(FVulkanDevice* Device, const uint8* OptionalPSOCacheData,VkGraphicsPipelineCreateInfo* PipelineInfo, const FGfxPipelineDesc* GfxEntry, const FVulkanRenderTargetLayout* RTLayout, TArrayView<uint32_t> VS, TArrayView<uint32_t> PS, size_t& AfterSize, FString* FailureMessageOUT = nullptr) { return VK_NULL_HANDLE; }
+	// Gives platform a chance to handle precompile of PSOs, returns VK_NULL_HANDLE if unsupported
+	static VkPipelineCache PrecompilePSO(FVulkanDevice* Device, const TArrayView<const uint8> ContextData, const TArrayView<const uint8> OptionalPSOCacheData, FGraphicsPipelineStateInitializer::EPSOPrecacheCompileType PSOCompileType, const VkGraphicsPipelineCreateInfo* PipelineInfo, const FGfxPipelineDesc* GfxEntry, const FVulkanRenderTargetLayout* RTLayout, TArrayView<const uint32_t> VS, TArrayView<const uint32_t> PS, size_t& AfterSize, FString* FailureMessageOUT = nullptr) { return VK_NULL_HANDLE;}
 
 	// Return VK_FALSE if platform wants to suppress the given debug report from the validation layers, VK_TRUE to print it.
 	static VkBool32 DebugReportFunction(VkDebugReportFlagsEXT MsgFlags, VkDebugReportObjectTypeEXT ObjType, uint64_t SrcObject, size_t Location, int32 MsgCode, const ANSICHAR* LayerPrefix, const ANSICHAR* Msg, void* UserData) { return VK_TRUE; }
@@ -160,4 +161,6 @@ public:
 	static VkShaderStageFlags RequiredWaveOpsShaderStageFlags(VkShaderStageFlags VulkanDeviceShaderStageFlags) { return VulkanDeviceShaderStageFlags; }
 
 	static VkTimeDomainKHR GetTimeDomain() { return VK_TIME_DOMAIN_DEVICE_KHR; }
+
+	static void OnCreateGfxPipelineCrash(struct FCrashContextExtendedWriter& Writer, bool bIsPrecache) {}
 };
diff --git a/Engine/Source/Runtime/VulkanRHI/Private/VulkanPipeline.cpp b/Engine/Source/Runtime/VulkanRHI/Private/VulkanPipeline.cpp
--- a/Engine/Source/Runtime/VulkanRHI/Private/VulkanPipeline.cpp
+++ b/Engine/Source/Runtime/VulkanRHI/Private/VulkanPipeline.cpp
@@ -127,6 +127,14 @@ static FAutoConsoleVariableRef GEnablePipelineCacheCompressionCvar(
 	ECVF_Default | ECVF_RenderThreadSafe
 );
 
+static bool GbEnableGfxPipelineCrashContext = false;
+static FAutoConsoleVariableRef GEnableGfxPipelineCrashContextCvar(
+	TEXT("r.Vulkan.EnableGfxPipelineCrashContext"),
+	GbEnableGfxPipelineCrashContext,
+	TEXT("Whether to include additional crash context information on crashes coming from GFX pipeline creation. (disabled by default)"),
+	ECVF_Default | ECVF_RenderThreadSafe
+);
+
 
 enum class ESingleThreadedPSOCreateMode
 {
@@ -1892,8 +1900,20 @@ static VkResult CreatePSOWithExternalService(FVulkanDevice* Device, FGraphicsPip
 		InitialCacheData.SetNumUninitialized(InitialCacheSize);
 		VulkanRHI::vkGetPipelineCacheData(Device->GetHandle(), DestPipelineCache, &InitialCacheSize, InitialCacheData.GetData());
 	}
+
+	TArray<uint8> HashBytes;
+	HashBytes.SetNumZeroed(sizeof(FSHAHash::Hash) * 2);
+	if(Shaders[ShaderStage::Vertex])
+	{
+		FMemory::Memcpy(HashBytes.GetData(), &static_cast<FVulkanVertexShader*>(Shaders[ShaderStage::Vertex])->GetHash().Hash[0], sizeof(FSHAHash::Hash));
+	}
+	if (Shaders[ShaderStage::Pixel])
+	{
+		FMemory::Memcpy(HashBytes.GetData()+sizeof(FSHAHash::Hash), &static_cast<FVulkanPixelShader*>(Shaders[ShaderStage::Pixel])->GetHash().Hash[0], sizeof(FSHAHash::Hash));
+	}
+
 	FString FailLog;
-	LocalPipelineCache = FVulkanPlatform::PrecompilePSO(Device, InitialCacheData, PSOCompileType, &PipelineInfo, GfxEntry, &PSO->RenderPass->GetLayout(), VSCode, PSCode, AfterSize,&FailLog);
+	LocalPipelineCache = FVulkanPlatform::PrecompilePSO(Device, HashBytes, InitialCacheData, PSOCompileType, &PipelineInfo, GfxEntry, &PSO->RenderPass->GetLayout(), VSCode, PSCode, AfterSize,&FailLog);
 
 	if (ensure(LocalPipelineCache != VK_NULL_HANDLE))
 	{
@@ -1915,6 +1935,37 @@ static VkResult CreatePSOWithExternalService(FVulkanDevice* Device, FGraphicsPip
 
 VkResult FVulkanPipelineStateCacheManager::CreateVKPipeline(FVulkanGraphicsPipelineState* PSO, FVulkanShader* Shaders[ShaderStage::NumGraphicsStages], const VkGraphicsPipelineCreateInfo& PipelineInfo, FGraphicsPipelineStateInitializer::EPSOPrecacheCompileType PSOCompileType)
 {
+#if WITH_ADDITIONAL_CRASH_CONTEXTS
+	// during a crash write out the shader hashes and give the platform's pipeline crash handler an opportunity.
+	uint32 tid = FPlatformTLS::GetCurrentThreadId();
+	auto AddToCrash = [Shaders, tid, PSOCompileType](FCrashContextExtendedWriter& Writer)
+		{
+			bool bIsPrecompileJob = PSOCompileType != FGraphicsPipelineStateInitializer::EPSOPrecacheCompileType::NotSet;
+			FVulkanPlatform::OnCreateGfxPipelineCrash(Writer, bIsPrecompileJob);
+
+			// write hash strings out, trying to avoid allocating
+			auto HashToString = [](const FSHAHash& ShaderHash, TCHAR(&Output)[MAX_SPRINTF])
+				{
+					static_assert(UE_ARRAY_COUNT(ShaderHash.Hash) * 2 < MAX_SPRINTF);
+					UE::String::BytesToHex(MakeArrayView(ShaderHash.Hash, UE_ARRAY_COUNT(ShaderHash.Hash)), Output);
+					Output[UE_ARRAY_COUNT(ShaderHash.Hash) * 2] = '\0';
+					return true;
+				};
+
+			TCHAR CrashString[MAX_SPRINTF] = { 0 };
+			FSHAHash VSHash = (Shaders[ShaderStage::Vertex]) ? static_cast<FVulkanVertexShader*>(Shaders[ShaderStage::Vertex])->GetHash() : FSHAHash();
+			FSHAHash PSHash = (Shaders[ShaderStage::Pixel]) ? static_cast<FVulkanPixelShader*>(Shaders[ShaderStage::Pixel])->GetHash() : FSHAHash();
+			TCHAR VSHashString[MAX_SPRINTF] = { 0 };
+			TCHAR PSHashString[MAX_SPRINTF] = { 0 };
+			HashToString(VSHash, VSHashString);
+			HashToString(PSHash, PSHashString);
+			FCString::Snprintf(CrashString, UE_ARRAY_COUNT(CrashString), TEXT("VS: %s, PS: %s (tid:%d, pt %d)"), VSHashString, PSHashString, tid, (int)PSOCompileType);
+			Writer.AddString(TEXT("VKhashes"), CrashString);
+		};
+
+	UE_ADD_CRASH_CONTEXT_SCOPE((GbEnableGfxPipelineCrashContext ? MoveTemp(AddToCrash) : TUniqueFunction<void(FCrashContextExtendedWriter&)>([](FCrashContextExtendedWriter&) {})));
+#endif
+
 	bool bIsPrecompileJob = PSOCompileType != FGraphicsPipelineStateInitializer::EPSOPrecacheCompileType::NotSet;
 	if(FVulkanChunkedPipelineCacheManager::IsEnabled())
 	{
@@ -2016,7 +2067,7 @@ VkResult FVulkanPipelineStateCacheManager::CreateVKPipeline(FVulkanGraphicsPipel
 		TArrayView<uint32_t> PSCode = PS.GetCodeView();
 		size_t AfterSize = 0;
 		FString FailLog;
-		LocalPipelineCache = FVulkanPlatform::PrecompilePSO(Device, MakeArrayView<uint8>(nullptr,0), PSOCompileType, &PipelineInfo, GfxEntry, &PSO->RenderPass->GetLayout(), VSCode, PSCode, AfterSize,&FailLog);
+		LocalPipelineCache = FVulkanPlatform::PrecompilePSO(Device, MakeArrayView<uint8>(nullptr,0), MakeArrayView<uint8>(nullptr, 0), PSOCompileType, &PipelineInfo, GfxEntry, &PSO->RenderPass->GetLayout(), VSCode, PSCode, AfterSize,&FailLog);
 
 		if (ensure(LocalPipelineCache != VK_NULL_HANDLE))
 		{
