VulkanRHI:  Add specialization constant support to Android PSO service.  Add ATTACHMENT_DESCRIPTION_STENCIL_LAYOUT support to PSO Service.

#rnx

#Synced-CL 49034046

diff --git a/Engine/Build/Android/Prebuilt/PSOService/Source/PSOService.cpp b/Engine/Build/Android/Prebuilt/PSOService/Source/PSOService.cpp
--- a/Engine/Build/Android/Prebuilt/PSOService/Source/PSOService.cpp
+++ b/Engine/Build/Android/Prebuilt/PSOService/Source/PSOService.cpp
@@ -40,27 +40,30 @@
 #define LOG_VERBOSE(...) __android_log_print(ANDROID_LOG_VERBOSE, APPNAME, ## __VA_ARGS__)
 #endif
 
-namespace UE::Jni::PsoServices
-{
-	struct FPsoProgramService
-	{
-		static constexpr const char* ClassName = "com/epicgames/unreal/psoservices/PSOProgramService";
+namespace UE { 
+	namespace Jni {
+		namespace PsoServices {
+			struct FPsoProgramService
+			{
+				static constexpr const char* ClassName = "com/epicgames/unreal/psoservices/PSOProgramService";
 
-		static void JNICALL NativeSetThreadPriority(JNIEnv* env, jobject thiz, jlong PriInfoIn);
-		static void JNICALL InitVKDevice(JNIEnv* env, jobject thiz);
-		static void JNICALL ShutdownVKDevice(JNIEnv* env, jobject thiz);
-		static jbyteArray JNICALL CompileVKGFXPSO(JNIEnv* env, jobject thiz, jbyteArray jVS, jbyteArray jPS, jbyteArray jPSO, jbyteArray jPSOCacheDataSource, jfloatArray jCompilationDuration);
-		static jint JNICALL CompileVKGFXPSOSHM(JNIEnv* env, jobject thiz, jint SHMemFD, jlong jVSSize, jlong jPSSize, jlong jPSOSize, jlong jPSOCacheDataSourceSize, jfloatArray jCompilationDuration);
+				static void JNICALL NativeSetThreadPriority(JNIEnv* env, jobject thiz, jlong PriInfoIn);
+				static void JNICALL InitVKDevice(JNIEnv* env, jobject thiz);
+				static void JNICALL ShutdownVKDevice(JNIEnv* env, jobject thiz);
+				static jbyteArray JNICALL CompileVKGFXPSO(JNIEnv* env, jobject thiz, jbyteArray jVS, jbyteArray jPS, jbyteArray jPSO, jbyteArray jPSOCacheDataSource, jfloatArray jCompilationDuration);
+				static jint JNICALL CompileVKGFXPSOSHM(JNIEnv* env, jobject thiz, jint SHMemFD, jlong jVSSize, jlong jPSSize, jlong jPSOSize, jlong jPSOCacheDataSourceSize, jfloatArray jCompilationDuration);
+			};
 
-		static constexpr JNINativeMethod NativeMethods[]
-		{
-			{"NativeSetThreadPriority", "(J)V"          , (void*)NativeSetThreadPriority},
-			{"InitVKDevice"           , "()V"           , (void*)InitVKDevice},
-			{"ShutdownVKDevice"       , "()V"           , (void*)ShutdownVKDevice},
-			{"CompileVKGFXPSO"        , "([B[B[B[B[F)[B", (void*)CompileVKGFXPSO},
-			{"CompileVKGFXPSOSHM"     , "(IJJJJ[F)I"    , (void*)CompileVKGFXPSOSHM}
-		};
-	};
+			static JNINativeMethod NativeMethods[] =
+			{
+				{"NativeSetThreadPriority", "(J)V"          , (void*)FPsoProgramService::NativeSetThreadPriority},
+				{"InitVKDevice"           , "()V"           , (void*)FPsoProgramService::InitVKDevice},
+				{"ShutdownVKDevice"       , "()V"           , (void*)FPsoProgramService::ShutdownVKDevice},
+				{"CompileVKGFXPSO"        , "([B[B[B[B[F)[B", (void*)FPsoProgramService::CompileVKGFXPSO},
+				{"CompileVKGFXPSOSHM"     , "(IJJJJ[F)I"    , (void*)FPsoProgramService::CompileVKGFXPSOSHM}
+			};
+		}
+	}
 }
 
 JNIEXPORT jint JNI_OnLoad(JavaVM* vm, void* reserved)
@@ -84,7 +87,7 @@ JNIEXPORT jint JNI_OnLoad(JavaVM* vm, void* reserved)
 		return JNI_ERR;
 	}
 
-	jint result = env->RegisterNatives(clazz, PsoServices::FPsoProgramService::NativeMethods, sizeof(PsoServices::FPsoProgramService::NativeMethods) / sizeof(JNINativeMethod));
+	jint result = env->RegisterNatives(clazz, PsoServices::NativeMethods, sizeof(PsoServices::NativeMethods) / sizeof(JNINativeMethod));
 
 	env->DeleteLocalRef(clazz);
 
@@ -398,7 +401,7 @@ public:
 			CreateInfo.pNext = &FSRCreateInfo;
 		}
 
-		VkPipelineShaderStageCreateInfo ShaderStages[2];
+		std::vector<VkPipelineShaderStageCreateInfo> ShaderStages(PipelineCreateInfo.StageCount);
 
 		// VkPipelineShaderStageCreateInfo
 		for (int32_t Idx = 0; Idx < PipelineCreateInfo.StageCount; ++Idx)
@@ -421,8 +424,22 @@ public:
 			ShaderStages[Idx].pName = (const char*)&PSO[MemoryOffset];
 			ShaderStages[Idx].pNext = PipelineShaderStageCreatePNext;
 			MemoryOffset += NameLength;
+
+			bool bHasSpecializationInfo;
+			COPY_FROM_BUFFER(&bHasSpecializationInfo, PSO, MemoryOffset, sizeof(bHasSpecializationInfo));
+			if (bHasSpecializationInfo)
+			{
+				ShaderStages[Idx].pSpecializationInfo = (VkSpecializationInfo*)&PSO[MemoryOffset];
+				MemoryOffset += sizeof(VkSpecializationInfo);
+
+				((VkSpecializationInfo*)ShaderStages[Idx].pSpecializationInfo)->pMapEntries = (VkSpecializationMapEntry*)&PSO[MemoryOffset];
+				MemoryOffset += ShaderStages[Idx].pSpecializationInfo->mapEntryCount * sizeof(VkSpecializationMapEntry);
+
+				((VkSpecializationInfo*)ShaderStages[Idx].pSpecializationInfo)->pData = (void*)&PSO[MemoryOffset];
+				MemoryOffset += ShaderStages[Idx].pSpecializationInfo->dataSize;
+			}
 		}
-		CreateInfo.pStages = ShaderStages;
+		CreateInfo.pStages = ShaderStages.data();
 
 		VkPipelineVertexInputStateCreateInfo VertexInputState;
 		if (PipelineCreateInfo.bHasVkPipelineVertexInputStateCreateInfo)
@@ -596,6 +613,17 @@ public:
 			{
 				RenderPassCreateInfo.pAttachments = (VkAttachmentDescription2KHR*)&PSO[MemoryOffset];
 				MemoryOffset += sizeof(VkAttachmentDescription2KHR) * RenderPassCreateInfo.attachmentCount;
+
+				for (uint32_t AttachmentIndex = 0; AttachmentIndex < RenderPassCreateInfo.attachmentCount; ++AttachmentIndex)
+				{
+					bool bHasStencilLayout = false;
+					COPY_FROM_BUFFER(&bHasStencilLayout, PSO, MemoryOffset, sizeof(bool));;
+					if (bHasStencilLayout)
+					{
+						const_cast<VkAttachmentDescription2*>(&RenderPassCreateInfo.pAttachments[AttachmentIndex])->pNext = &PSO[MemoryOffset];
+						MemoryOffset += sizeof(VkAttachmentDescriptionStencilLayout);
+					}
+				}
 			}
 
 			if (RenderPassCreateInfo.dependencyCount > 0)
@@ -1071,19 +1099,19 @@ jbyteArray JNICALL UE::Jni::PsoServices::FPsoProgramService::CompileVKGFXPSO(JNI
 
 	if (PSOCacheDataSource)
 	{
-		jenv->ReleaseByteArrayElements(jPSOCacheDataSource, PSOCacheDataSource, JNI_ABORT);
+		jenv->ReleaseByteArrayElements(jPSOCacheDataSource, (int8_t*)PSOCacheDataSource, JNI_ABORT);
 	}
 	if (PSO)
 	{
-		jenv->ReleaseByteArrayElements(jPSO, PSO, JNI_ABORT);
+		jenv->ReleaseByteArrayElements(jPSO, (int8_t*)PSO, JNI_ABORT);
 	}
 	if (PS)
 	{
-		jenv->ReleaseByteArrayElements(jPS, PS, JNI_ABORT);
+		jenv->ReleaseByteArrayElements(jPS, (int8_t*)PS, JNI_ABORT);
 	}
 	if (VS)
 	{
-		jenv->ReleaseByteArrayElements(jVS, VS, JNI_ABORT);
+		jenv->ReleaseByteArrayElements(jVS, (int8_t*)VS, JNI_ABORT);
 	}
 
 	uint32_t Size = 0;
diff --git a/Engine/Source/Runtime/VulkanRHI/Private/Android/VulkanAndroidPlatform.cpp b/Engine/Source/Runtime/VulkanRHI/Private/Android/VulkanAndroidPlatform.cpp
--- a/Engine/Source/Runtime/VulkanRHI/Private/Android/VulkanAndroidPlatform.cpp
+++ b/Engine/Source/Runtime/VulkanRHI/Private/Android/VulkanAndroidPlatform.cpp
@@ -1199,6 +1199,28 @@ void HandleDepthStencilAttachmentPNext(const VkAttachmentReference2* Attachment,
 	check(HandledCount == VkStructs.Num());
 }
 
+void HandleAttachmentDescriptionPNext(const VkAttachmentDescription2& Description, TArray<char>& MemoryStream)
+{
+	TMap<VkStructureType, const void*> VkStructs;
+	GetVKStructsFromPNext(Description.pNext, VkStructs, { VK_STRUCTURE_TYPE_ATTACHMENT_DESCRIPTION_STENCIL_LAYOUT });
+
+	int32_t HandledCount = 0;
+
+	const void** Struct = VkStructs.Find(VK_STRUCTURE_TYPE_ATTACHMENT_DESCRIPTION_STENCIL_LAYOUT);
+	const bool bHasStencilLayout = Struct != nullptr;
+	COPY_TO_BUFFER(MemoryStream, &bHasStencilLayout, sizeof(bool));
+
+	if (bHasStencilLayout)
+	{
+		VkAttachmentDescriptionStencilLayout* StencilLayout = (VkAttachmentDescriptionStencilLayout*)*Struct;
+		check(StencilLayout->pNext == nullptr);
+		COPY_TO_BUFFER(MemoryStream, StencilLayout, sizeof(VkAttachmentDescriptionStencilLayout));
+		HandledCount++;
+	}
+
+	check(HandledCount == VkStructs.Num());
+}
+
 void PipelineToBinary(FVulkanDevice* Device, const VkGraphicsPipelineCreateInfo* PipelineInfo, const FGfxPipelineDesc* GfxEntry, const FVulkanRenderTargetLayout* RTLayout, TArray<char>& MemoryStream)
 {
 	static const unsigned int INITIAL_PSO_STREAM_SIZE = 64 * 1024;
@@ -1249,6 +1271,21 @@ void PipelineToBinary(FVulkanDevice* Device, const VkGraphicsPipelineCreateInfo*
 
 		COPY_TO_BUFFER(MemoryStream, &NameLength, sizeof(uint32_t));
 		COPY_TO_BUFFER(MemoryStream, PipelineInfo->pStages[Idx].pName, NameLength);
+
+		const bool bHasSpecializationInfo = PipelineInfo->pStages[Idx].pSpecializationInfo != nullptr;
+		COPY_TO_BUFFER(MemoryStream, &bHasSpecializationInfo, sizeof(bHasSpecializationInfo));
+		if (bHasSpecializationInfo)
+		{
+			const VkSpecializationInfo SpecializationInfo = {
+				.mapEntryCount = PipelineInfo->pStages[Idx].pSpecializationInfo->mapEntryCount,
+				.pMapEntries = nullptr,
+				.dataSize = PipelineInfo->pStages[Idx].pSpecializationInfo->dataSize,
+				.pData = nullptr
+			};
+			COPY_TO_BUFFER(MemoryStream, &SpecializationInfo, sizeof(VkSpecializationInfo));
+			COPY_TO_BUFFER(MemoryStream, PipelineInfo->pStages[Idx].pSpecializationInfo->pMapEntries, SpecializationInfo.mapEntryCount * sizeof(VkSpecializationMapEntry));
+			COPY_TO_BUFFER(MemoryStream, PipelineInfo->pStages[Idx].pSpecializationInfo->pData, SpecializationInfo.dataSize);
+		}
 	}
 
 	if (pipelineCreateInfo.bHasVkPipelineVertexInputStateCreateInfo)
@@ -1448,6 +1485,11 @@ void PipelineToBinary(FVulkanDevice* Device, const VkGraphicsPipelineCreateInfo*
 		if(RenderPassCreateInfo.attachmentCount > 0)
 		{
 			COPY_TO_BUFFER(MemoryStream, CreateInfo.pAttachments, sizeof(VkAttachmentDescription2) * RenderPassCreateInfo.attachmentCount);
+
+			for (uint32_t AttachmentIndex = 0; AttachmentIndex < RenderPassCreateInfo.attachmentCount; ++AttachmentIndex)
+			{
+				HandleAttachmentDescriptionPNext(CreateInfo.pAttachments[AttachmentIndex], MemoryStream);
+			}
 		}
 
 		if(RenderPassCreateInfo.dependencyCount > 0)
@@ -1854,7 +1896,7 @@ FString FVulkanAndroidPlatform::GetCompiledPSOCacheTopFolderPath()
 }
 
 VkPipelineCache FVulkanAndroidPlatform::PrecompilePSO(
-		FVulkanDevice* Device,
+		FVulkanDevice* Device, 
 		const TArrayView<const uint8> ContextData,
 		const TArrayView<const uint8> OptionalPSOCacheData, 
 		FGraphicsPipelineStateInitializer::EPSOPrecacheCompileType PSOCompileType,
