Protection from memory stats issues

#android

diff --git a/Engine/Build/Android/Java/src/com/epicgames/unreal/GameActivity.java.template b/Engine/Build/Android/Java/src/com/epicgames/unreal/GameActivity.java.template
--- a/Engine/Build/Android/Java/src/com/epicgames/unreal/GameActivity.java.template
+++ b/Engine/Build/Android/Java/src/com/epicgames/unreal/GameActivity.java.template
@@ -3673,30 +3673,37 @@ Trace.endSection();  // WriteCache
 
 	protected String[] FindLineFromFile(String Filepath, String[] fieldsToReport)
 	{
-		String[] foundStrings = new String[fieldsToReport.length];
-		File file = new File(Filepath);
-		if (file.exists())
+		try
 		{
-			try (BufferedReader reader = new BufferedReader(new FileReader(file));)
+			String[] foundStrings = new String[fieldsToReport.length];
+			File file = new File(Filepath);
+			if (file.exists())
 			{
-				String line;
-				while ((line = reader.readLine()) != null)
+				try (BufferedReader reader = new BufferedReader(new FileReader(file));)
 				{
-					for(int i = 0; i<fieldsToReport.length; i++)
+					String line;
+					while ((line = reader.readLine()) != null)
 					{
-						if(line.startsWith(fieldsToReport[i]))
+						for(int j = 0; j<fieldsToReport.length; j++)
 						{
-							foundStrings[i] = line;
+							if(line.startsWith(fieldsToReport[j]))
+							{
+								foundStrings[j] = line;
+							}
 						}
 					}
 				}
+				catch (IOException ie)
+				{
+					Log.debug("failed to read status: " + ie);
+				}
 			}
-			 catch (IOException ie)
-			 {
-				Log.debug("failed to read status: " + ie);
-			}
+			return foundStrings;
+		}
+		catch (Throwable t)
+		{
+			return null;
 		}
-		return foundStrings;
 	}
 	
 	private static Activity getActivityOfContext(Context context)
@@ -4157,10 +4164,27 @@ Trace.endSection();
 				}
 
 				String[] statusLines = FindLineFromFile("/proc/self/status",new String[] {"VmSize:","VmHWM:", "VmRSS:", "VmData:", "VmSwap:"});
-				String rssSubString = statusLines[2].substring(0, statusLines[2].lastIndexOf(' '));
-				int rssMemory = Integer.parseInt(rssSubString.substring(rssSubString.lastIndexOf(' ')+1, rssSubString.length()));
+				int rssMemory = 0;
+				if (statusLines != null && statusLines[2] != null)
+				{
+					int lastSpace = statusLines[2].lastIndexOf(' ');
+					if (lastSpace >= 0)
+					{
+						String rssSubString = statusLines[2].substring(0, lastSpace);
+						lastSpace = rssSubString.lastIndexOf(' ');
+						if (lastSpace >= 0)
+						{
+							try
+							{
+								rssMemory = Integer.parseInt(rssSubString.substring(lastSpace+1, rssSubString.length()));
+							}
+							catch (NumberFormatException ne)
+							{
+							}
+						}
+					}
+				}
 				int rssDiffed = Math.abs(rssMemory - LastRSSMemoryValue);
-
 				String ProcessMemInfoLine = "Used memory: ";
 				if(rssDiffed >= memstatRSSRequiredDiff)
 				{
@@ -4172,18 +4196,30 @@ Trace.endSection();
 						LastProcessMemoryValue = ProcessMemory;
 					}
 					memLog += "memstat: "+ProcessMemInfoLine + ProcessMemory + "\n";
-					for(String memstatusline : statusLines)
+					if (statusLines != null)
 					{
-						memLog += "memstat: "+memstatusline + "\n";
+						for(String memstatusline : statusLines)
+						{
+							if (memstatusline != null)
+							{
+								memLog += "memstat: "+memstatusline + "\n";
+							}
+						}
 					}
 				
 					if( showPSSInfo )
 					{
 						String[] smapsLines = FindLineFromFile("/proc/self/smaps_rollup",new String[] {"SwapPss:","Pss:"});
-				    	for(String memstatusline : smapsLines)
-				    	{
-							memLog += "memstat: "+memstatusline + "\n";
-				    	}
+						if (smapsLines != null)
+						{
+			    			for(String memstatusline : smapsLines)
+			    			{
+								if (memstatusline != null)
+								{
+									memLog += "memstat: "+memstatusline + "\n";
+								}
+			    			}
+						}
 					}
 //$${gameActivityMemStatAdditions}$$
 					Log.debug(memLog);
