Improve thread safety on SystemInfoLogs

#android

diff --git a/Engine/Build/Android/Java/src/com/epicgames/unreal/GameActivity.java.template b/Engine/Build/Android/Java/src/com/epicgames/unreal/GameActivity.java.template
--- a/Engine/Build/Android/Java/src/com/epicgames/unreal/GameActivity.java.template
+++ b/Engine/Build/Android/Java/src/com/epicgames/unreal/GameActivity.java.template
@@ -1533,10 +1533,7 @@ Trace.endSection();
 			}
 			catch (IOException ie)
 			{
-				if (SystemInfoLogs != null)
-				{
-					SystemInfoLogs.add("failed to read /proc/cpuinfo: " + ie);
-				}
+				SystemInfo_Log("failed to read /proc/cpuinfo: " + ie);
 			}
 			if (reader != null)
 			{
@@ -1551,10 +1548,7 @@ Trace.endSection();
 		}
 		else
 		{
-			if (SystemInfoLogs != null)
-			{
-				SystemInfoLogs.add("failed to open /proc/cpuinfo!");
-			}
+			SystemInfo_Log("failed to open /proc/cpuinfo!");
 		}
 		return cpuInfo;
 	}
@@ -1590,10 +1584,7 @@ Trace.endSection();
 		}
 		catch (Exception e)
 		{
-			if (SystemInfoLogs != null)
-			{
-				SystemInfoLogs.add("Unable to use getprop: " + e);
-			}
+			SystemInfo_Log("Unable to use getprop: " + e);
 		}
 		if (reader != null)
 		{
@@ -2482,12 +2473,12 @@ Trace.endSection();
 		// try to get chip info
 		// use getprop ro.hardware as fallback (not useful on Qualcomm, just returns qcom)
 		String roHardware = getProp("ro.hardware");
-		SystemInfoLogs.add("ro.hardware: " + roHardware);
+		SystemInfo_Log("ro.hardware: " + roHardware);
 		if (roHardware.equals("qcom"))
 		{
 			// use ro.hardware.chipname for Qualcomm
 			String roHardwareChipname = getProp("ro.hardware.chipname");
-			SystemInfoLogs.add("ro.hardware.chipname: " + roHardwareChipname);
+			SystemInfo_Log("ro.hardware.chipname: " + roHardwareChipname);
 			roHardware = roHardwareChipname;
 			if (roHardware.isEmpty())
 			{
@@ -2500,11 +2491,11 @@ Trace.endSection();
 		{
 			hardware = CpuInfo.get("Hardware");
 		}
-		SystemInfoLogs.add("Hardware: " + hardware);
+		SystemInfo_Log("Hardware: " + hardware);
 		String roSOCModel = getProp("ro.soc.model");
 		if (!roSOCModel.equals(""))
 		{
-			SystemInfoLogs.add("ro.soc.model: " + roSOCModel);
+			SystemInfo_Log("ro.soc.model: " + roSOCModel);
 
 			// validate either we only have ro.soc.model, or that previous hardware value in agreement
 			// encountered a device which had ro.hardware and /proc/cpuinfo correct but ro.soc.model set differently
@@ -2516,7 +2507,7 @@ Trace.endSection();
 			}
 			else
 			{
-				SystemInfoLogs.add("Warning: hardware does not match ro.soc.model, assuming hardware is correct");
+				SystemInfo_Log("Warning: hardware does not match ro.soc.model, assuming hardware is correct");
 			}
 		}
 
@@ -2529,7 +2520,7 @@ Trace.endSection();
 		activityManager.getMemoryInfo(memoryInfo);
 
 		int totalMemory = (int)(memoryInfo.totalMem / 1024 / 1024);
-		SystemInfoLogs.add("Memory: " + totalMemory + " MB");
+		SystemInfo_Log("Memory: " + totalMemory + " MB");
 		variables.put("memory", Integer.toString(totalMemory));
 
 		//variables.put("Profile", "Android");
