Fixing Virtual Keyboard's EditText being transparent on certain Android Device.
- Re-added surface view's Transparent background when Virtual Keyboard is open and reset back to Black (default) when closed.

#Synced-CL 50732174

diff --git a/Engine/Build/Android/Java/src/com/epicgames/unreal/GameActivity.java.template b/Engine/Build/Android/Java/src/com/epicgames/unreal/GameActivity.java.template
--- a/Engine/Build/Android/Java/src/com/epicgames/unreal/GameActivity.java.template
+++ b/Engine/Build/Android/Java/src/com/epicgames/unreal/GameActivity.java.template
@@ -552,6 +552,8 @@ public class GameActivity extends $${gameActivitySuperClass}$$ implements Surfac
 	private SurfaceView MySurfaceView;
 	private int DesiredHolderWidth = 0;
 	private int DesiredHolderHeight = 0;
+	private boolean bSurfaceViewIsTransparent = false;
+	protected static int SurfaceViewDefaultBackgroundColor = Color.BLACK;
 
 	/** Discovered Vulkan Version and Level from getSystemAvailableFeatures() */
 	public int VulkanVersion = 0;
@@ -592,6 +594,11 @@ public class GameActivity extends $${gameActivitySuperClass}$$ implements Surfac
 		Keyboard
 	}
 	private EAlertDialogType CurrentDialogType = EAlertDialogType.None;
+	
+	public static void SetSurfaceViewDefaultBackgroundColor(int color)
+	{
+		SurfaceViewDefaultBackgroundColor = color;
+	}
 
 	private void SystemInfo_InitLog()
 	{
@@ -4880,6 +4887,11 @@ Trace.beginSection("UI Building");
 								newVirtualKeyboardInput.getLayoutParams().width = Math.abs(visibleRect.right - visibleRect.left);
 								newVirtualKeyboardInput.setX(leftDiff);
 								newVirtualKeyboardInput.setY(keyboardYPos);
+								if (!bSurfaceViewIsTransparent)
+								{
+									bSurfaceViewIsTransparent = true;
+									MySurfaceView.setBackgroundColor(Color.TRANSPARENT);
+								}
 							}
 
 							AndroidThunkJava_ResizeKeyboard(true, keyboardRect, visibleRect);
@@ -4897,6 +4909,12 @@ Trace.beginSection("UI Building");
 
 									//set offscreen
 									newVirtualKeyboardInput.setY(-1000);
+									
+									if (bSurfaceViewIsTransparent)
+                                    {
+                                    	bSurfaceViewIsTransparent = false;
+                                    	MySurfaceView.setBackgroundColor(SurfaceViewDefaultBackgroundColor);
+                                    }
 								}
 							}
 
