[Backout] - CL50981786
#rnx
Original CL Desc
-----------------------------------------------------------------
Reimplemented Android Background HTTP Manager

#android

diff --git a/Engine/Build/Android/Java/src/com/epicgames/unreal/GameActivity.java.template b/Engine/Build/Android/Java/src/com/epicgames/unreal/GameActivity.java.template
--- a/Engine/Build/Android/Java/src/com/epicgames/unreal/GameActivity.java.template
+++ b/Engine/Build/Android/Java/src/com/epicgames/unreal/GameActivity.java.template
@@ -7529,7 +7529,7 @@ Trace.endSection();
 		{
 			NetworkListener = new NetworkConnectivityClient.Listener() {
 				@Override
-				public void onNetworkAvailable(NetworkConnectivityClient.NetworkTransportType networkTransportType, boolean isNetworkMetered)
+				public void onNetworkAvailable(NetworkConnectivityClient.NetworkTransportType networkTransportType)
 				{
 					nativeNetworkChanged();
 				}
diff --git a/Engine/Build/Android/Java/src/com/epicgames/unreal/network/NetworkChangedManager.java b/Engine/Build/Android/Java/src/com/epicgames/unreal/network/NetworkChangedManager.java
--- a/Engine/Build/Android/Java/src/com/epicgames/unreal/network/NetworkChangedManager.java
+++ b/Engine/Build/Android/Java/src/com/epicgames/unreal/network/NetworkChangedManager.java
@@ -2,10 +2,14 @@
 
 package com.epicgames.unreal.network;
 
+import android.annotation.TargetApi;
 import android.content.Context;
 import android.net.ConnectivityManager;
 import android.net.Network;
 import android.net.NetworkCapabilities;
+import android.net.NetworkInfo;
+import android.net.NetworkRequest;
+import android.os.Build;
 import android.os.Handler;
 import androidx.annotation.NonNull;
 import androidx.annotation.Nullable;
@@ -17,7 +21,9 @@ import java.lang.ref.WeakReference;
 import java.net.HttpURLConnection;
 import java.net.MalformedURLException;
 import java.net.URL;
-import java.util.concurrent.CopyOnWriteArrayList;
+import java.util.HashSet;
+import java.util.Iterator;
+import java.util.Set;
 import java.util.concurrent.ExecutorService;
 import java.util.concurrent.Executors;
 
@@ -41,7 +47,7 @@ public final class NetworkChangedManager implements NetworkConnectivityClient {
 
 	private static NetworkChangedManager instance;
 	@NonNull
-	public static synchronized NetworkChangedManager getInstance() {
+	public static synchronized NetworkConnectivityClient getInstance() {
 		if (instance == null) {
 			instance = new NetworkChangedManager();
 		}
@@ -54,14 +60,12 @@ public final class NetworkChangedManager implements NetworkConnectivityClient {
 	private ConnectivityManager connectivityManager;
 
 	@NonNull
-	private final CopyOnWriteArrayList<WeakReference<Listener>> networkChangedListeners = new CopyOnWriteArrayList<>();
+	private Set<WeakReference<Listener>> networkChangedListeners = new HashSet<>();
 
 	/*
 	 * Data
 	 */
-	private boolean isNetworkAvailable;
-	private boolean isNetworkMetered;
-	private boolean isNetworkBlocked;
+	private Set<String> networks = new HashSet<>();
 	@Nullable
 	private ConnectivityState currentState = null;
 	@Nullable
@@ -78,53 +82,80 @@ public final class NetworkChangedManager implements NetworkConnectivityClient {
 
 	@Override
 	public void initNetworkCallback(@NonNull Context context) {
+		if (Build.VERSION.SDK_INT < Build.VERSION_CODES.LOLLIPOP) {
+			// unsupported before Lollipop
+			return;
+		}
+
 		connectivityManager = (ConnectivityManager) context.getApplicationContext().getSystemService(Context.CONNECTIVITY_SERVICE);
 		if (connectivityManager == null) {
 			Log.error("Unable to start connectivityManager");
 			return;
 		}
-		connectivityManager.registerDefaultNetworkCallback(connectivityListener);
+		NetworkRequest.Builder builder = new NetworkRequest.Builder();
+		builder.addCapability(NetworkCapabilities.NET_CAPABILITY_INTERNET);
+		builder.addTransportType(NetworkCapabilities.TRANSPORT_WIFI);
+		builder.addTransportType(NetworkCapabilities.TRANSPORT_ETHERNET);
+		builder.addTransportType(NetworkCapabilities.TRANSPORT_VPN);
+		builder.addTransportType(NetworkCapabilities.TRANSPORT_CELLULAR);
+		builder.addTransportType(NetworkCapabilities.TRANSPORT_BLUETOOTH);
+		if (Build.VERSION.SDK_INT >= Build.VERSION_CODES.M) {
+			builder.addCapability(NetworkCapabilities.NET_CAPABILITY_VALIDATED);
+		}
+		connectivityManager.registerNetworkCallback(builder.build(), connectivityListener);
 	}
 
+	@TargetApi(Build.VERSION_CODES.LOLLIPOP)
 	private ConnectivityManager.NetworkCallback connectivityListener = new ConnectivityManager.NetworkCallback() {
+		@Override
+		public void onAvailable(Network network) {
+			networks.add(network.toString());
+			Log.verbose("Network Available: " + network.toString());
+			checkNetworkConnectivity();
+		}
+
 		@Override
 		public void onLost(Network network) {
-			isNetworkAvailable = false;
+			networks.remove(network.toString());
 			Log.verbose("Network Lost callback: " + network.toString());
+
+			if (networks.isEmpty()) {
+				Log.verbose("All Networks Lost");
+			} else {
+				Log.verbose("Not All Networks Lost");
+			}
 			checkNetworkConnectivity();
 		}
 
 		@Override
 		public void onCapabilitiesChanged(Network network, NetworkCapabilities networkCapabilities) {
-			if (!networkCapabilities.hasCapability(NetworkCapabilities.NET_CAPABILITY_VALIDATED)) {
-				isNetworkAvailable = false;
+			if (Build.VERSION.SDK_INT >= Build.VERSION_CODES.M
+					&& !networkCapabilities.hasCapability(NetworkCapabilities.NET_CAPABILITY_VALIDATED)) {
 				Log.verbose("Network Capabilities changed, doesn't have validated net_capability");
+				networks.remove(network.toString());
+			} else if (!networkCapabilities.hasCapability(NetworkCapabilities.NET_CAPABILITY_INTERNET)) {
+				Log.verbose("Network Capabilities changed, has Internet: " + false);
+				networks.remove(network.toString());
 			} else {
-				isNetworkAvailable = networkCapabilities.hasCapability(NetworkCapabilities.NET_CAPABILITY_INTERNET);
-				Log.verbose("Network Capabilities changed, has Internet: " + isNetworkAvailable);
+				Log.verbose("Network Capabilities changed, has Internet: " + true);
+				networks.add(network.toString());
 			}
 			checkNetworkConnectivity();
 		}
+	};
 
+	private Runnable retryRunnable = new Runnable() {
 		@Override
-		public void onBlockedStatusChanged(@NonNull Network network, boolean blocked)
-		{
-			isNetworkBlocked = blocked;
-			Log.verbose("Network Blocked Status changed, Blocked: " + isNetworkBlocked);
+		public void run() {
+			Log.verbose("Attempting to check for network connectivity again.");
+			retryCount++;
+			retryScheduled = false;
 			checkNetworkConnectivity();
 		}
 	};
 
-	private Runnable retryRunnable = () ->
-	{
-		Log.verbose("Attempting to check for network connectivity again.");
-		retryCount++;
-		retryScheduled = false;
-		checkNetworkConnectivity();
-	};
-
 	private void scheduleRetry() {
-		if (networkChangedListeners.isEmpty()) {
+		if (networkChangedListeners.size() == 0) {
 			Log.verbose("No listeners so not retrying. When a listener is added the connection will be checked.");
 			return;
 		}
@@ -152,8 +183,7 @@ public final class NetworkChangedManager implements NetworkConnectivityClient {
 
 	private void setNetworkState(ConnectivityState state) {
         NetworkTransportType updatedNetworkTransportType = calculateNetworkTransport(connectivityManager);
-		boolean updateIsNetworkMetered = connectivityManager.isActiveNetworkMetered();
-        if (currentState == state && currentNetworkTransport == updatedNetworkTransportType && isNetworkMetered == updateIsNetworkMetered) {
+        if (currentState == state && currentNetworkTransport == updatedNetworkTransportType) {
 			Log.verbose("Connectivity hasn't changed. Current state: " + currentState);
 			if (currentState != ConnectivityState.CONNECTION_AVAILABLE) {
 				scheduleRetry();
@@ -162,8 +192,7 @@ public final class NetworkChangedManager implements NetworkConnectivityClient {
 		}
 		currentState = state;
         currentNetworkTransport = updatedNetworkTransportType;
-		isNetworkMetered = updateIsNetworkMetered;
-		fireNetworkChangeListeners(state, currentNetworkTransport, isNetworkMetered);
+		fireNetworkChangeListeners(state, currentNetworkTransport);
 		Log.verbose("Network connectivity changed. New connectivity state: " + state);
 
 		if (currentState != ConnectivityState.CONNECTION_AVAILABLE) {
@@ -189,74 +218,88 @@ public final class NetworkChangedManager implements NetworkConnectivityClient {
 
 		networkCheckInProgress = true;
 		final ExecutorService executor = Executors.newSingleThreadExecutor();
-		final Runnable timeoutRunnable = () ->
-		{
-			Log.verbose("Unable to connect to: " + getCurrentHostResolutionAddress());
-			networkCheckInProgress = false;
-			executor.shutdownNow();
-			setNetworkState(ConnectivityState.NO_CONNECTION);
-			scheduleRetry();
+		final Runnable timeoutRunnable = new Runnable() {
+			@Override
+			public void run() {
+				Log.verbose("Unable to connect to: " + getCurrentHostResolutionAddress());
+				networkCheckInProgress = false;
+				executor.shutdownNow();
+				setNetworkState(ConnectivityState.NO_CONNECTION);
+				scheduleRetry();
+			}
 		};
 		internalScheduler.postDelayed(timeoutRunnable, HOSTNAME_RESOLUTION_TIMEOUT_MS * HOST_RESOLUTION_ADDRESSES.length + 100);
-		executor.execute(() ->
-		{
-			android.net.TrafficStats.setThreadStatsTag(1);
-			HttpURLConnection urlConnection = null;
-			boolean connectedSuccessfully = false;
-			// Attempt to connect to any of the hostnames. If any succeed we are connected to
-			// the internet.
-			for (int i = 0; i < HOST_RESOLUTION_ADDRESSES.length; i++) {
-				try {
-					Log.verbose("Verifying internet connection with host: " + getCurrentHostResolutionAddress());
-					URL url = new URL(getCurrentHostResolutionAddress());
-					urlConnection = (HttpURLConnection) url.openConnection();
-					urlConnection.setUseCaches(false);
-					urlConnection.setRequestMethod("HEAD");
-					urlConnection.setConnectTimeout((int) HOSTNAME_RESOLUTION_TIMEOUT_MS);
-					urlConnection.setReadTimeout((int) HOSTNAME_RESOLUTION_TIMEOUT_MS);
-					urlConnection.getInputStream().close();
-					connectedSuccessfully = true;
-					break;
-				} catch (MalformedURLException e) {
-					Log.error("Malformed URL, this should never happen. Please fix, url: " + getCurrentHostResolutionAddress());
-				} catch (IOException e) {
-					Log.verbose("Unable to connect to: " + getCurrentHostResolutionAddress());
-				} catch (Exception e) {
-					Log.verbose("Unable to connect to: " + getCurrentHostResolutionAddress() + ", exception: " + e.toString());
-				} finally {
-					if (urlConnection != null) {
-						urlConnection.disconnect();
-					}
-					if (!connectedSuccessfully) {
-						/*
-						 * Move to the next resolution address.
-						 * When we try again we should do it with the next available URL. For
-						 * some reason on some networks we've had issues with
-						 * https://example.com so using https://google.com and others will give
-						 * us more redundancy
-						 */
-						nextHostResolutionAddress();
+		executor.execute(new Runnable() {
+			@Override
+			public void run() {
+				android.net.TrafficStats.setThreadStatsTag(1);
+				HttpURLConnection urlConnection = null;
+				boolean connectedSuccessfully = false;
+				// Attempt to connect to any of the hostnames. If any succeed we are connected to
+				// the internet.
+				for (int i = 0; i < HOST_RESOLUTION_ADDRESSES.length; i++) {
+					try {
+						Log.verbose("Verifying internet connection with host: " + getCurrentHostResolutionAddress());
+						URL url = new URL(getCurrentHostResolutionAddress());
+						urlConnection = (HttpURLConnection) url.openConnection();
+						urlConnection.setUseCaches(false);
+						urlConnection.setRequestMethod("HEAD");
+						urlConnection.setConnectTimeout((int) HOSTNAME_RESOLUTION_TIMEOUT_MS);
+						urlConnection.setReadTimeout((int) HOSTNAME_RESOLUTION_TIMEOUT_MS);
+						urlConnection.getInputStream().close();
+						connectedSuccessfully = true;
+						break;
+					} catch (MalformedURLException e) {
+						Log.error("Malformed URL, this should never happen. Please fix, url: " + getCurrentHostResolutionAddress());
+					} catch (IOException e) {
+						Log.verbose("Unable to connect to: " + getCurrentHostResolutionAddress());
+					} catch (Exception e) {
+						Log.verbose("Unable to connect to: " + getCurrentHostResolutionAddress() + ", exception: " + e.toString());
+					} finally {
+						if (urlConnection != null) {
+							urlConnection.disconnect();
+						}
+						if (!connectedSuccessfully) {
+							/*
+							 * Move to the next resolution address.
+							 * When we try again we should do it with the next available URL. For
+							 * some reason on some networks we've had issues with
+							 * https://example.com so using https://google.com and others will give
+							 * us more redundancy
+							 */
+							nextHostResolutionAddress();
+						}
 					}
 				}
-			}
 
-			internalScheduler.removeCallbacks(timeoutRunnable);
-			networkCheckInProgress = false;
-			if (connectedSuccessfully) {
-				setNetworkState(ConnectivityState.CONNECTION_AVAILABLE);
-			} else {
-				setNetworkState(ConnectivityState.NO_CONNECTION);
-				scheduleRetry();
-			}
-			Log.verbose("Full network check complete. State: " + currentState);
+				internalScheduler.removeCallbacks(timeoutRunnable);
+				networkCheckInProgress = false;
+				if (connectedSuccessfully) {
+					setNetworkState(ConnectivityState.CONNECTION_AVAILABLE);
+				} else {
+					setNetworkState(ConnectivityState.NO_CONNECTION);
+					scheduleRetry();
+				}
+				Log.verbose("Full network check complete. State: " + currentState);
 
-			executor.shutdownNow();
+				executor.shutdownNow();
+			}
 		});
 	}
 
 	@NonNull
 	private ConnectivityState calculateNetworkConnectivityNaively() {
-		return isNetworkAvailable && !isNetworkBlocked ? ConnectivityState.CONNECTION_AVAILABLE : ConnectivityState.NO_CONNECTION;
+		if (networks.isEmpty() || connectivityManager == null) {
+			return ConnectivityState.NO_CONNECTION;
+		} else {
+			NetworkInfo activeNetworkInfo = connectivityManager.getActiveNetworkInfo();
+
+			if (activeNetworkInfo != null && activeNetworkInfo.isAvailable() && activeNetworkInfo.isConnected()) {
+				return ConnectivityState.CONNECTION_AVAILABLE;
+			} else {
+				return ConnectivityState.NO_CONNECTION;
+			}
+		}
 	}
 
 	private String getCurrentHostResolutionAddress() {
@@ -278,6 +321,11 @@ public final class NetworkChangedManager implements NetworkConnectivityClient {
 
 	@Override
 	public boolean addListener(Listener listener, boolean fireImmediately) {
+		if (Build.VERSION.SDK_INT < Build.VERSION_CODES.LOLLIPOP) {
+			// unsupported before Lollipop
+			return false;
+		}
+
 		for (WeakReference<Listener> currentListener : networkChangedListeners) {
 			if (currentListener.get() == listener) {
 				return false;
@@ -293,69 +341,64 @@ public final class NetworkChangedManager implements NetworkConnectivityClient {
 		}
 		// Current state will never be null after calling checkNetworkConnectivity
 		if (fireImmediately && currentState != null) {
-			fireNetworkChangeListenerInternal(listener, currentState, currentNetworkTransport, isNetworkMetered);
+			fireNetworkChangeListenerInternal(listener, currentState, currentNetworkTransport);
 		}
 		return true;
 	}
 
 	@Override
 	public boolean removeListener(Listener listener) {
-		if (!networkChangedListeners.removeIf(reference ->
-			{
-				Listener item = reference.get();
-				return item == listener || item == null;
-			}))
-		{
+		if (Build.VERSION.SDK_INT < Build.VERSION_CODES.LOLLIPOP) {
+			// unsupported before Lollipop
 			return false;
 		}
 
-		if (networkChangedListeners.isEmpty())
-		{
+		Iterator<WeakReference<Listener>> changeListenersIterator = networkChangedListeners.iterator();
+		while (changeListenersIterator.hasNext()) {
+			WeakReference<Listener> currentListener = changeListenersIterator.next();
+			if (currentListener.get() == listener) {
+				changeListenersIterator.remove();
+				return true;
+			}
+		}
+		/*
+		 * If there are no listeners, clear any retries. When a listener is added we'll check the
+		 * connection again.
+		 */
+		if (networkChangedListeners.size() == 0) {
 			clearRetry();
 		}
-
-		return true;
+		return false;
 	}
 
-	private void fireNetworkChangeListeners(ConnectivityState state, NetworkTransportType networkTransportType, boolean isNetworkMetered) {
+	private void fireNetworkChangeListeners(ConnectivityState state, NetworkTransportType networkTransportType) {
 		// can be called from callbacks that should return quickly
 		final ExecutorService executor = Executors.newSingleThreadExecutor();
-		executor.execute(() ->
-		{
-			boolean needsCleanup = false;
-
-			for (WeakReference<Listener> reference : networkChangedListeners)
-			{
-				Listener listener = reference.get();
-				if (listener == null)
-				{
-					needsCleanup = true;
-				}
-				else
-				{
-					fireNetworkChangeListenerInternal(listener, state, networkTransportType, isNetworkMetered);
-				}
-			}
-
-			if (needsCleanup)
-			{
-				if (networkChangedListeners.removeIf(reference -> reference.get() == null) && networkChangedListeners.isEmpty())
-				{
-					clearRetry();
+		executor.execute(new Runnable() {
+			@Override
+			public void run() {
+				Iterator<WeakReference<Listener>> changeListenersIterator = networkChangedListeners.iterator();
+				while (changeListenersIterator.hasNext()) {
+					WeakReference<Listener> listenerWeakReference = changeListenersIterator.next();
+					Listener listener = listenerWeakReference.get();
+					if (listener == null) {
+						changeListenersIterator.remove();
+					} else {
+						fireNetworkChangeListenerInternal(listener, state, networkTransportType);
+					}
 				}
+				executor.shutdownNow();
 			}
-
-			executor.shutdownNow();
 		});
 	}
 
-	private void fireNetworkChangeListenerInternal(Listener listener, ConnectivityState state, NetworkTransportType networkTransportType, boolean isNetworkMetered) {
+	private void fireNetworkChangeListenerInternal(Listener listener, ConnectivityState state, NetworkTransportType networkTransportType) {
 		switch (state) {
 			case NO_CONNECTION:
 				listener.onNetworkLost();
 				break;
 			case CONNECTION_AVAILABLE:
-				listener.onNetworkAvailable(networkTransportType, isNetworkMetered);
+				listener.onNetworkAvailable(networkTransportType);
 				break;
 		}
 	}
@@ -365,34 +408,56 @@ public final class NetworkChangedManager implements NetworkConnectivityClient {
 	{
 		return calculateNetworkTransport(connectivityManager);
 	}
-
+	
 	@Override
 	public void checkConnectivity() {
 		checkNetworkConnectivity();
 	}
 
     private NetworkTransportType calculateNetworkTransport(@NonNull ConnectivityManager connectivityManager) {
-		Network activeNetwork = connectivityManager.getActiveNetwork();
-		if (activeNetwork == null) {
-			return NetworkTransportType.UNKNOWN;
-		}
-		NetworkCapabilities capabilities = connectivityManager.getNetworkCapabilities(activeNetwork);
-		if (capabilities == null) {
-			return NetworkTransportType.UNKNOWN;
-		}
-
-		if (capabilities.hasTransport(NetworkCapabilities.TRANSPORT_WIFI)) {
-			return NetworkTransportType.WIFI;
-		} else if (capabilities.hasTransport(NetworkCapabilities.TRANSPORT_CELLULAR)) {
-			return NetworkTransportType.CELLULAR;
-		} else if (capabilities.hasTransport(NetworkCapabilities.TRANSPORT_ETHERNET)) {
-			return NetworkTransportType.ETHERNET;
-		} else if (capabilities.hasTransport(NetworkCapabilities.TRANSPORT_BLUETOOTH)) {
-			return NetworkTransportType.BLUETOOTH;
-		} else if (capabilities.hasTransport(NetworkCapabilities.TRANSPORT_VPN)) {
-			return NetworkTransportType.VPN;
-		} else {
-			return NetworkTransportType.UNKNOWN;
-		}
-	}
+       if (Build.VERSION.SDK_INT < Build.VERSION_CODES.M) {
+           NetworkInfo networkInfo = connectivityManager.getActiveNetworkInfo();
+           if (networkInfo == null) {
+               return NetworkTransportType.UNKNOWN;
+           }
+
+           switch (networkInfo.getType()) {
+               case ConnectivityManager.TYPE_WIFI:
+                   return NetworkTransportType.WIFI;
+               case ConnectivityManager.TYPE_MOBILE:
+                   return NetworkTransportType.CELLULAR;
+               case ConnectivityManager.TYPE_ETHERNET:
+                   return NetworkTransportType.ETHERNET;
+               case ConnectivityManager.TYPE_BLUETOOTH:
+                   return NetworkTransportType.BLUETOOTH;
+               case ConnectivityManager.TYPE_VPN:
+                   return NetworkTransportType.VPN;
+               default:
+                   return NetworkTransportType.UNKNOWN;
+           }
+       } else {
+           Network activeNetwork = connectivityManager.getActiveNetwork();
+           if (activeNetwork == null) {
+               return NetworkTransportType.UNKNOWN;
+           }
+           NetworkCapabilities capabilities = connectivityManager.getNetworkCapabilities(activeNetwork);
+           if (capabilities == null) {
+               return NetworkTransportType.UNKNOWN;
+           }
+
+           if (capabilities.hasTransport(NetworkCapabilities.TRANSPORT_WIFI)) {
+               return NetworkTransportType.WIFI;
+           } else if (capabilities.hasTransport(NetworkCapabilities.TRANSPORT_CELLULAR)) {
+               return NetworkTransportType.CELLULAR;
+           } else if (capabilities.hasTransport(NetworkCapabilities.TRANSPORT_ETHERNET)) {
+               return NetworkTransportType.ETHERNET;
+           } else if (capabilities.hasTransport(NetworkCapabilities.TRANSPORT_BLUETOOTH)) {
+               return NetworkTransportType.BLUETOOTH;
+           } else if (capabilities.hasTransport(NetworkCapabilities.TRANSPORT_VPN)) {
+               return NetworkTransportType.VPN;
+           } else {
+               return NetworkTransportType.UNKNOWN;
+           }
+       }
+   }
 }
diff --git a/Engine/Build/Android/Java/src/com/epicgames/unreal/network/NetworkConnectivityClient.java b/Engine/Build/Android/Java/src/com/epicgames/unreal/network/NetworkConnectivityClient.java
--- a/Engine/Build/Android/Java/src/com/epicgames/unreal/network/NetworkConnectivityClient.java
+++ b/Engine/Build/Android/Java/src/com/epicgames/unreal/network/NetworkConnectivityClient.java
@@ -3,21 +3,20 @@
 package com.epicgames.unreal.network;
 
 import android.content.Context;
-
 import androidx.annotation.NonNull;
 
 public interface NetworkConnectivityClient {
-	enum NetworkTransportType {
-		WIFI,
-		VPN,
-		ETHERNET,
-		CELLULAR,
-		BLUETOOTH,
-		UNKNOWN
-	}
-
-	interface Listener {
-		void onNetworkAvailable(NetworkTransportType networkTransportType, boolean isNetworkMetered);
+    public enum NetworkTransportType {
+        WIFI,
+        VPN,
+        ETHERNET,
+        CELLULAR,
+        BLUETOOTH,
+        UNKNOWN
+    }
+
+    public interface Listener {
+        void onNetworkAvailable(NetworkTransportType networkTransportType);
 
 		void onNetworkLost();
 	}
diff --git a/Engine/Plugins/Online/Android/AndroidFetchBackgroundDownload/Source/AndroidFetchBackgroundDownload/AndroidFetchBackgroundDownload_UPL.xml b/Engine/Plugins/Online/Android/AndroidFetchBackgroundDownload/Source/AndroidFetchBackgroundDownload/AndroidFetchBackgroundDownload_UPL.xml
--- a/Engine/Plugins/Online/Android/AndroidFetchBackgroundDownload/Source/AndroidFetchBackgroundDownload/AndroidFetchBackgroundDownload_UPL.xml
+++ b/Engine/Plugins/Online/Android/AndroidFetchBackgroundDownload/Source/AndroidFetchBackgroundDownload/AndroidFetchBackgroundDownload_UPL.xml
@@ -45,12 +45,6 @@
 		</insert>
 	</buildGradleAdditions>
 
-	<androidManifestUpdates>
-		<addElements tag="application">
-			<receiver android:name="com.epicgames.unreal.download.DownloadManager$CancelReceiver" android:exported="false"/>
-		</addElements>
-	</androidManifestUpdates>
-	
 	<!-- optional additions to proguard -->
 	<proguardAdditions>
 		<insert>
@@ -58,10 +52,6 @@
 
 			# Android Background HTTP Classes with native calls ...
 
-			-keep class com.epicgames.unreal.download.DownloadManager {
-			public *;
-			}
-
 			-keep class com.epicgames.unreal.download.UEDownloadWorker {
 			public *;
 			}
@@ -83,12 +73,5 @@
 	<!-- set build version to BranchName-Changelist, see C++ equivalent in AndroidPlatformBackgroundHttpRequest.cpp-->
 	<gameActivityClassAdditions>
 		<insertValue value="public static final String BGDL_BUILD_VERSION = &quot;$S(EngineBranchName)&quot;;"/>
-		<insertNewLine/>
 	</gameActivityClassAdditions>
-
-	<gameActivityOnResumeAdditions>
-		<insert>
-			com.epicgames.unreal.download.DownloadManager.Resume();
-		</insert>
-	</gameActivityOnResumeAdditions>
 </root>
diff --git a/Engine/Plugins/Online/Android/AndroidFetchBackgroundDownload/Source/AndroidFetchBackgroundDownload/Java/download/DownloadManager.java b/Engine/Plugins/Online/Android/AndroidFetchBackgroundDownload/Source/AndroidFetchBackgroundDownload/Java/download/DownloadManager.java
deleted file mode 100644
diff --git a/Engine/Plugins/Online/Android/AndroidFetchBackgroundDownload/Source/AndroidFetchBackgroundDownload/Java/download/DownloadProgressListener.java b/Engine/Plugins/Online/Android/AndroidFetchBackgroundDownload/Source/AndroidFetchBackgroundDownload/Java/download/DownloadProgressListener.java
--- a/Engine/Plugins/Online/Android/AndroidFetchBackgroundDownload/Source/AndroidFetchBackgroundDownload/Java/download/DownloadProgressListener.java
+++ b/Engine/Plugins/Online/Android/AndroidFetchBackgroundDownload/Source/AndroidFetchBackgroundDownload/Java/download/DownloadProgressListener.java
@@ -1,18 +1,15 @@
 // Copyright Epic Games, Inc. All Rights Reserved.
 package com.epicgames.unreal.download;
 
-import androidx.annotation.NonNull;
-
 import com.epicgames.unreal.download.UEDownloadWorker.EDownloadCompleteReason;
-import com.epicgames.unreal.download.datastructs.DownloadDescription;
 
 //Interface for a class that wants to be notified when a download completes
 public interface DownloadProgressListener
 {
-	void OnDownloadProgress(@NonNull String DestinationLocation, long TotalBytesWritten, long BytesWrittenSinceLastCall);
-	void OnDownloadGroupProgress(int GroupID, int Progress, boolean Indeterminate);
-	void OnDownloadComplete(@NonNull String DestinationLocation, EDownloadCompleteReason CompleteReason);
-	void OnDownloadMetrics(@NonNull String DestinationLocation, long TotalBytesDownloaded, long DownloadDuration, long DownloadStartTimeUTC, long DownloadEndTimeUTC);
-	void OnAllDownloadsComplete(boolean bDidAllRequestsSucceed);
-	void OnDownloadEnqueued(@NonNull DownloadDescription DownloadDescription);
+	public void OnDownloadProgress(String DestinationLocation, long TotalBytesWritten, long BytesWrittenSinceLastCall);
+	public void OnDownloadGroupProgress(int GroupID, int Progress, boolean Indeterminate);
+	public void OnDownloadComplete(String DestinationLocation, EDownloadCompleteReason CompleteReason);
+	public void OnDownloadMetrics(String DestinationLocation, long TotalBytesDownloaded, long DownloadDuration, long DownloadStartTimeUTC, long DownloadEndTimeUTC);
+	public void OnAllDownloadsComplete(boolean bDidAllRequestsSucceed);
+	public void OnDownloadEnqueued(String DestinationLocation, boolean bEnqueueSuccess);
 }
\ No newline at end of file
diff --git a/Engine/Plugins/Online/Android/AndroidFetchBackgroundDownload/Source/AndroidFetchBackgroundDownload/Java/download/UEDownloadWorker.java b/Engine/Plugins/Online/Android/AndroidFetchBackgroundDownload/Source/AndroidFetchBackgroundDownload/Java/download/UEDownloadWorker.java
--- a/Engine/Plugins/Online/Android/AndroidFetchBackgroundDownload/Source/AndroidFetchBackgroundDownload/Java/download/UEDownloadWorker.java
+++ b/Engine/Plugins/Online/Android/AndroidFetchBackgroundDownload/Source/AndroidFetchBackgroundDownload/Java/download/UEDownloadWorker.java
@@ -26,7 +26,6 @@ import java.io.File;
 import com.epicgames.unreal.GameActivity;
 
 import com.epicgames.unreal.Logger;
-import com.epicgames.unreal.download.datastructs.DownloadDescription;
 import com.epicgames.unreal.workmanager.UEWorker;
 
 import com.epicgames.unreal.download.datastructs.DownloadNotificationDescription;
@@ -158,7 +157,7 @@ public class UEDownloadWorker extends UEWorker implements DownloadProgressListen
 		{
 			NetworkListener = new NetworkConnectivityClient.Listener() {
 				@Override
-				public void onNetworkAvailable(NetworkConnectivityClient.NetworkTransportType networkTransportType, boolean isNetworkMetered) {
+				public void onNetworkAvailable(NetworkConnectivityClient.NetworkTransportType networkTransportType) {
 					bLostNetwork = false;
 					if (mFetchManager != null) {
 						mFetchManager.NetworkAvailable();
@@ -783,7 +782,7 @@ public class UEDownloadWorker extends UEWorker implements DownloadProgressListen
 	// DownloadCompletionListener Implementation
 	//
 	@Override
-	public void OnDownloadProgress(@NonNull String DestinationLocation, long TotalBytesWritten, long BytesWrittenSinceLastCall)
+	public void OnDownloadProgress(String DestinationLocation, long TotalBytesWritten, long BytesWrittenSinceLastCall)
 	{
 		if (!bGameThreadIsActive)
 		{
@@ -813,7 +812,7 @@ public class UEDownloadWorker extends UEWorker implements DownloadProgressListen
 	}
 	
 	@Override
-	public void OnDownloadComplete(@NonNull String DestinationLocation, EDownloadCompleteReason CompleteReason)
+	public void OnDownloadComplete(String DestinationLocation, EDownloadCompleteReason CompleteReason)
 	{
 		if (!bGameThreadIsActive)
 		{
@@ -832,7 +831,7 @@ public class UEDownloadWorker extends UEWorker implements DownloadProgressListen
 	}
 
 	@Override
-	public void OnDownloadMetrics(@NonNull String DestinationLocation, long TotalBytesDownloaded, long DownloadDuration, long DownloadStartTimeUTC, long DownloadEndTimeUTC)
+	public void OnDownloadMetrics(String DestinationLocation, long TotalBytesDownloaded, long DownloadDuration, long DownloadStartTimeUTC, long DownloadEndTimeUTC)
 	{
 		if (!bGameThreadIsActive)
 		{
@@ -901,9 +900,17 @@ public class UEDownloadWorker extends UEWorker implements DownloadProgressListen
 	}
 
 	@Override
-	public void OnDownloadEnqueued(@NonNull DownloadDescription DownloadDescription)
+	public void OnDownloadEnqueued(String DestinationLocation, boolean bEnqueueSuccess)
 	{
-		Log.verbose("Enqueue success:" + DownloadDescription.DestinationLocation);
+		if (bEnqueueSuccess)
+		{
+			Log.verbose("Enqueue success:" + DestinationLocation);
+		}
+		else
+		{
+			Log.debug("Enqueue failure, retrying request:" + DestinationLocation);
+			mFetchManager.RetryDownload(DestinationLocation);
+		}
 
 		bHasEnqueueHappened = true;
 	}
diff --git a/Engine/Plugins/Online/Android/AndroidFetchBackgroundDownload/Source/AndroidFetchBackgroundDownload/Java/download/asyncdownloader/AsyncDownloadRequest.java b/Engine/Plugins/Online/Android/AndroidFetchBackgroundDownload/Source/AndroidFetchBackgroundDownload/Java/download/asyncdownloader/AsyncDownloadRequest.java
--- a/Engine/Plugins/Online/Android/AndroidFetchBackgroundDownload/Source/AndroidFetchBackgroundDownload/Java/download/asyncdownloader/AsyncDownloadRequest.java
+++ b/Engine/Plugins/Online/Android/AndroidFetchBackgroundDownload/Source/AndroidFetchBackgroundDownload/Java/download/asyncdownloader/AsyncDownloadRequest.java
@@ -13,19 +13,19 @@ import java.util.Objects;
 final public class AsyncDownloadRequest {
     public final String Id;								// internal unique Id used for tracking
     public final String RequestId;						// optional request Id used for deduplication
-    public final String[] Urls;							// cycle through these on failure
+    public final List<String> Urls;						// cycle through these on failure
     public final int MaxRetries;						// total retry attempts across all URLs
     public final File TargetFile;
     public final boolean bResumeIfPossible;				// use Range: bytes=<currentLen>- if file exists, enabled by default
     public final @Nullable Map<String, String> Headers; // optional request headers per attempt
 	public final int GroupID;
-	public long TotalSize;
+	public final long TotalSize;
 
     private AsyncDownloadRequest(Builder b) {
         this.Id = Objects.requireNonNull(b.Id, "id");
         this.RequestId = b.RequestId;
-        this.Urls = b.Urls;
-        if (this.Urls.length == 0) 
+        this.Urls = List.copyOf(Objects.requireNonNull(b.Urls, "urls"));
+        if (this.Urls.isEmpty()) 
 			throw new IllegalArgumentException("urls must not be empty");
         this.MaxRetries = Math.max(0, b.MaxRetries);
         this.TargetFile = Objects.requireNonNull(b.TargetFile, "targetFile");
@@ -36,13 +36,13 @@ final public class AsyncDownloadRequest {
     }
 
     public static class Builder {
-		private String Id;
-		private String RequestId;
-		private String[] Urls = new String[0];
-		private int MaxRetries = 3;
-		private File TargetFile;
-		private boolean bResumeIfPossible = true;
-		private Map<String, String> Headers;
+        private String Id;
+        private String RequestId;
+        private List<String> Urls = new ArrayList<>();
+        private int MaxRetries = 3;
+        private File TargetFile;
+        private boolean bResumeIfPossible = true;
+        private Map<String, String> Headers;
 		private int GroupID = 0;
 		private long TotalSize = -1L;
 
@@ -54,8 +54,8 @@ final public class AsyncDownloadRequest {
 			this.RequestId = RequestId; 
 			return this; 
 		}
-        public Builder SetUrls(String[] Urls) { 
-			this.Urls = Urls; 
+        public Builder SetUrls(List<String> Urls) { 
+			this.Urls = new ArrayList<>(Urls); 
 			return this; 
 		}
 
diff --git a/Engine/Plugins/Online/Android/AndroidFetchBackgroundDownload/Source/AndroidFetchBackgroundDownload/Java/download/asyncdownloader/AsyncDownloadRequestListener.java b/Engine/Plugins/Online/Android/AndroidFetchBackgroundDownload/Source/AndroidFetchBackgroundDownload/Java/download/asyncdownloader/AsyncDownloadRequestListener.java
--- a/Engine/Plugins/Online/Android/AndroidFetchBackgroundDownload/Source/AndroidFetchBackgroundDownload/Java/download/asyncdownloader/AsyncDownloadRequestListener.java
+++ b/Engine/Plugins/Online/Android/AndroidFetchBackgroundDownload/Source/AndroidFetchBackgroundDownload/Java/download/asyncdownloader/AsyncDownloadRequestListener.java
@@ -2,15 +2,15 @@
 
 package com.epicgames.unreal.download.asyncdownloader;
 
-import androidx.annotation.NonNull;
+import androidx.annotation.Nullable;
 
 import java.io.File;
 import java.net.URL;
 
 public interface AsyncDownloadRequestListener {
-    void OnStart(@NonNull String Id);
-	void OnProgress(@NonNull String Id, long BytesDownloadedSoFar);
-	void OnSuccess(@NonNull String Id, @NonNull File File);
-	void OnError(@NonNull String Id, @NonNull Throwable Error, URL LastTriedUrl);
-	void OnPaused(@NonNull String Id);
+    public void OnStart(String Id, URL Url, long StartingOffset, long TotalBytes);
+	public void OnProgress(String Id, long BytesDownloadedSoFar, long TotalBytes);
+	public void OnSuccess(String Id, File File);
+	public void OnError(String Id, Throwable Error, @Nullable URL LastTriedUrl);
+	public void OnPaused(String Id);
 }
diff --git a/Engine/Plugins/Online/Android/AndroidFetchBackgroundDownload/Source/AndroidFetchBackgroundDownload/Java/download/asyncdownloader/AsyncDownloader.java b/Engine/Plugins/Online/Android/AndroidFetchBackgroundDownload/Source/AndroidFetchBackgroundDownload/Java/download/asyncdownloader/AsyncDownloader.java
--- a/Engine/Plugins/Online/Android/AndroidFetchBackgroundDownload/Source/AndroidFetchBackgroundDownload/Java/download/asyncdownloader/AsyncDownloader.java
+++ b/Engine/Plugins/Online/Android/AndroidFetchBackgroundDownload/Source/AndroidFetchBackgroundDownload/Java/download/asyncdownloader/AsyncDownloader.java
@@ -2,7 +2,6 @@
 
 package com.epicgames.unreal.download.asyncdownloader;
 
-import androidx.annotation.NonNull;
 import androidx.annotation.Nullable;
 
 import okhttp3.*;
@@ -56,17 +55,26 @@ public class AsyncDownloader {
 		//Log.debug("Submit: Id=" + Request.Id);
 		DownloadEntry Existing = Entries.get(Request.Id);
 		if (Existing == null) {
-			DownloadEntry Entry = new DownloadEntry(Request, Listener);
-			Entries.put(Request.Id, Entry);
-			StartOrQueue(Entry);
-			return true;
+			if (Request.RequestId == null || Request.RequestId.isEmpty()) {
+				DownloadEntry Entry = new DownloadEntry(Request, Listener);
+				Entries.put(Request.Id, Entry);
+				StartOrQueue(Entry);
+				return true;
+			}
+			Existing = Entries.values().stream().filter(Entry -> Request.RequestId.equals(Entry.Request.RequestId)).findFirst().orElse(null);
+			if (Existing == null) {
+				DownloadEntry Entry = new DownloadEntry(Request, Listener);
+				Entries.put(Request.Id, Entry);
+				StartOrQueue(Entry);
+				return true;
+			}
 		}
 		Existing.Listener = Listener;
 		//Log.debug("Submit: pre-Existing, updated Listener");
 		return false;
 	}
 
-	public synchronized void Pause(String Id, boolean userPaused) {
+	public synchronized void Pause(String Id) {
 		//Log.debug("Pause: Id=" + Id);
 		DownloadEntry Entry = Entries.get(Id);
 		if (Entry == null) {
@@ -76,10 +84,7 @@ public class AsyncDownloader {
 			//Log.debug("Pause: Id=" + Id + " is already paused, ignoring call");
 			return;
 		}
-		if (userPaused)
-		{
-			Entry.UserPaused = true;
-		}
+		Entry.UserPaused = true;
 		Call c = Entry.currentCall;
 		if (c != null) {
 			c.cancel();
@@ -96,11 +101,11 @@ public class AsyncDownloader {
 		//Log.debug("PauseAll");
 		PausedAll.set(true);
 		for (DownloadEntry Entry : Entries.values()) {
-			Pause(Entry.Request.Id, false);
+			Pause(Entry.Request.Id);
 		}
 	}
 
-	public synchronized void Resume(String Id, boolean userResumed) {
+	public synchronized void Resume(String Id) {
 		//Log.debug("Resume: Id=" + Id);
 		DownloadEntry Entry = Entries.get(Id);
 		if (Entry == null) {
@@ -110,10 +115,7 @@ public class AsyncDownloader {
 			//Log.debug("Resume: Id=" + Id + " not paused, ignoring call");
 			return;
 		}
-		if (userResumed)
-		{
-			Entry.UserPaused = false;
-		}
+		Entry.UserPaused = false;
 		StartOrQueue(Entry);
 	}
 
@@ -124,7 +126,7 @@ public class AsyncDownloader {
 		//Log.debug("ResumeAll");
 		PausedAll.set(false);
 		for (DownloadEntry Entry : Entries.values()) {
-			Resume(Entry.Request.Id, false);
+			Resume(Entry.Request.Id);
 		}
 	}
 
@@ -159,7 +161,7 @@ public class AsyncDownloader {
 	public synchronized void Retry(String Id) {
 		//Log.debug("Retry: Id=" + Id);
 
-		DownloadEntry OldEntry = Entries.remove(Id);
+		DownloadEntry OldEntry = Entries.get(Id);
 		if (OldEntry == null) {
 			return;
 		}
@@ -174,6 +176,7 @@ public class AsyncDownloader {
 		}
 		OldEntry.State = EState.CANCELLED;
 
+		Entries.remove(Id);
 		DownloadEntry NewEntry = new DownloadEntry(Req, Listener);
 		NewEntry.UserPaused = false;
 		Entries.put(Req.Id, NewEntry);
@@ -191,7 +194,7 @@ public class AsyncDownloader {
 				continue;
 			}
 			if (e.Request.TotalSize <= 0) {
-				return -1f;
+				continue;
 			}
 
 			if (e.State == EState.COMPLETED) {
@@ -244,7 +247,7 @@ public class AsyncDownloader {
 			Files.createDirectories(file.getParentFile().toPath());
 			long CurrentLen = (file.exists() && Entry.Request.bResumeIfPossible) ? file.length() : 0L;
 
-			URL NextUrl = new URL(Entry.Request.Urls[Entry.UrlIndex % Entry.Request.Urls.length]);
+			URL NextUrl = new URL(Entry.Request.Urls.get(Entry.UrlIndex % Entry.Request.Urls.size()));
 			Request.Builder rb = new Request.Builder().url(NextUrl);
 			if (Entry.Request.Headers != null) {
 				for (Map.Entry<String, String> h : Entry.Request.Headers.entrySet()) {
@@ -252,25 +255,11 @@ public class AsyncDownloader {
 				}
 			}
 
-			if (CurrentLen > 0) {
-				if (CurrentLen < Entry.Request.TotalSize)
-				{
+			if (CurrentLen > 0 && Entry.Request.TotalSize > 0) {
+				if (CurrentLen <= Entry.Request.TotalSize) {
 					rb.header("Range", "bytes=" + CurrentLen + "-");
 					//Log.debug("StartOrQueue: requesting partial content for id " + Entry.Request.Id + " | Offset " + CurrentLen);
-				} else if (CurrentLen == Entry.Request.TotalSize) {
-					synchronized(this) {
-						if (Entries.remove(Entry.Request.Id) != null)
-						{
-							Entry.State = EState.COMPLETED;
-							Entry.DownloadedSoFar = CurrentLen;
-							Entry.Listener.OnStart(Entry.Request.Id);
-							Entry.Listener.OnProgress(Entry.Request.Id, CurrentLen);
-							Entry.Listener.OnSuccess(Entry.Request.Id, file);
-						}
-					}
-					CheckIfAllDownloadsCompleted();
-					return;
-				} else if (Entry.Request.TotalSize > 0) {
+				} else {
 					// be on the safe side: delete the file completely if it is bigger than we expected
 					Log.error("StartOrQueue: requesting full file because current file size is too big (" + CurrentLen + " > " + Entry.Request.TotalSize + ")");
 					CurrentLen = 0;
@@ -281,18 +270,17 @@ public class AsyncDownloader {
 			Request Request = rb.build();
 			Entry.State = EState.RUNNING;
 			Entry.DownloadedSoFar = CurrentLen;
-			Entry.Listener.OnStart(Entry.Request.Id);
+			Entry.Listener.OnStart(Entry.Request.Id, NextUrl, CurrentLen, Entry.Request.TotalSize);
 
 			Call Call = Client.newCall(Request);
 			Entry.currentCall = Call;
-			final Object syncRoot = this;
 			Call.enqueue(new Callback() {
-				@Override public void onFailure(@NonNull Call Call, @NonNull IOException Except) {
-					Log.error("onFailure: error: failed to begin downloading | Id=" + Entry.Request.Id + " | Exception=" + Except);
+				@Override public void onFailure(Call Call, IOException Except) {
+					Log.error("onFailure: error: failed to begin downloading | Id=" + Entry.Request.Id);
 					HandleAttemptFailure(Entry, Except, NextUrl);
 				}
 
-				@Override public void onResponse(@NonNull Call Call, @NonNull Response Response) {
+				@Override public void onResponse(Call Call, Response Response) {
 					//Log.debug("onResponse: begin downloading | Id=" + Entry.Request.Id);
 					try (Response response = Response) {
 						int Code = response.code();
@@ -324,13 +312,12 @@ public class AsyncDownloader {
 						}
 
 						// check if the total matches what we expect
-						long FileOnDiskLen;
-						long TotalContentLen;
+						long TotalContentLen = 0;
+						long FileOnDiskLen = 0;
 						if (bShouldAppend) {
 							FileOnDiskLen = file.length();
-							TotalContentLen = ContentLen >= 0 ? FileOnDiskLen + ContentLen : ContentLen;
+							TotalContentLen = FileOnDiskLen + ContentLen;
 						} else {
-							FileOnDiskLen = 0;
 							TotalContentLen = ContentLen;
 						}
 						if (TotalContentLen != Entry.Request.TotalSize && Entry.Request.TotalSize > 0) {
@@ -340,42 +327,50 @@ public class AsyncDownloader {
 									NextUrl);
 							return;
 						}
-						if (TotalContentLen > 0) {
-							Entry.Request.TotalSize = TotalContentLen;
-						}
 
 						// all is good, start downloading
 						try (BufferedSource Source = Objects.requireNonNull(response.body()).source();
-							 FileChannel Channel = FileChannel.open(file.toPath(), StandardOpenOption.CREATE, StandardOpenOption.WRITE)) {
+							 FileChannel Channel = FileChannel.open(file.toPath(),
+									 StandardOpenOption.CREATE,
+									 StandardOpenOption.WRITE)) {
 
-							long WritePos = FileOnDiskLen;
+							long WritePos = bShouldAppend ? file.length() : 0L;
 							if (!bShouldAppend) {
 								Channel.truncate(0L);
 							}
 
-							while (!Source.exhausted()) {
-								long NumBytesRead = Channel.transferFrom(Source, WritePos, 64 * 1024);
+							byte[] Buffer = new byte[8192];
+							long TotalExpected = (ContentLen >= 0)
+									? (bShouldAppend ? WritePos + ContentLen : ContentLen)
+									: -1L;
 
+							while (true) {
+								int NumBytesRead = Source.read(Buffer, 0, Buffer.length);
+								if (NumBytesRead == -1) {
+									break;
+								}
+
+								Channel.position(WritePos);
+								Channel.write(java.nio.ByteBuffer.wrap(Buffer, 0, NumBytesRead));
 								WritePos += NumBytesRead;
 								Entry.DownloadedSoFar = WritePos;
-								Entry.Listener.OnProgress(Entry.Request.Id, WritePos);
+								Entry.Listener.OnProgress(Entry.Request.Id, WritePos, TotalExpected);
+							}
+
+							//Log.debug("onResponse: download completed | Id=" + Entry.Request.Id);
+							synchronized(this) {
+								if (Entries.remove(Entry.Request.Id) != null)
+								{ 
+									Entry.State = EState.COMPLETED;
+									Entry.currentCall = null;
+									Entry.Listener.OnSuccess(Entry.Request.Id, file);
+								}
 							}
+							CheckIfAllDownloadsCompleted();
 						}
 					} catch (Throwable Error) {
 						HandleAttemptFailure(Entry, Error, NextUrl);
-						return;
 					}
-
-					//Log.debug("onResponse: download completed | Id=" + Entry.Request.Id);
-					synchronized(syncRoot) {
-						if (Entries.remove(Entry.Request.Id) != null)
-						{
-							Entry.State = EState.COMPLETED;
-							Entry.currentCall = null;
-							Entry.Listener.OnSuccess(Entry.Request.Id, file);
-						}
-					}
-					CheckIfAllDownloadsCompleted();
 				}
 			});
 
@@ -384,7 +379,7 @@ public class AsyncDownloader {
 		}
 	}
 
-	private synchronized void HandleAttemptFailure(@NonNull DownloadEntry Entry, @NonNull Throwable Error, URL LastUrl) {
+	private synchronized void HandleAttemptFailure(DownloadEntry Entry, Throwable Error, @Nullable URL LastUrl) {
 		//Log.debug("HandleAttemptFailure: Id=" + Entry.Request.Id);
 		Entry.currentCall = null;
 		if (Entry.UserPaused || PausedAll.get()) {
@@ -396,7 +391,7 @@ public class AsyncDownloader {
 		}
 
 		int Attempt = ++Entry.Attempts;
-		Entry.UrlIndex = (Entry.UrlIndex + 1) % Entry.Request.Urls.length;
+		Entry.UrlIndex = (Entry.UrlIndex + 1) % Entry.Request.Urls.size();
 
 		final boolean bHaveRetriesRemaining = Attempt <= Entry.Request.MaxRetries;
 		if (bHaveRetriesRemaining) {
@@ -405,11 +400,9 @@ public class AsyncDownloader {
 		} else {
 			//Log.debug("HandleAttemptFailure: Id=" + Entry.Request.Id + " | out of retries: failing download");
 
-			if (Entries.remove(Entry.Request.Id) != null)
-			{
-				Entry.State = EState.FAILED;
-				Entry.Listener.OnError(Entry.Request.Id, Error, LastUrl);
-			}
+			Entry.State = EState.FAILED;
+			Entry.Listener.OnError(Entry.Request.Id, Error, LastUrl);
+			Entries.remove(Entry.Request.Id);
 			CheckIfAllDownloadsCompleted();
 		}
 	}
diff --git a/Engine/Plugins/Online/Android/AndroidFetchBackgroundDownload/Source/AndroidFetchBackgroundDownload/Java/download/datastructs/DownloadDescription.java b/Engine/Plugins/Online/Android/AndroidFetchBackgroundDownload/Source/AndroidFetchBackgroundDownload/Java/download/datastructs/DownloadDescription.java
--- a/Engine/Plugins/Online/Android/AndroidFetchBackgroundDownload/Source/AndroidFetchBackgroundDownload/Java/download/datastructs/DownloadDescription.java
+++ b/Engine/Plugins/Online/Android/AndroidFetchBackgroundDownload/Source/AndroidFetchBackgroundDownload/Java/download/datastructs/DownloadDescription.java
@@ -12,9 +12,8 @@ import java.io.IOException;
 import java.io.FileWriter;
 import java.io.FileReader;
 
-import java.util.ArrayList;
 import java.util.List;
-import java.util.stream.IntStream;
+import java.util.ArrayList;
 
 import com.epicgames.unreal.download.DownloadProgressListener;
 
@@ -24,7 +23,7 @@ public class DownloadDescription
 	//
 	//Serialized Members -- Pulled from WorkerParameters
 	//
-	public String[] URLs;
+	public List<String> URLs;
 	public String RequestID = null;
 	public String DestinationLocation = null;
 	public int RequestPriority = 0;
@@ -62,26 +61,9 @@ public class DownloadDescription
 	public long DownloadPauseTime = 0;
 	public long TotalPausedTime = 0;
 
-	// Fields used for final metrics
-	public long DownloadDuration = 0;
-	public long DownloadStartTimeUTC = 0;
-	public long DownloadEndTimeUTC = 0;
-	
 	public DownloadDescription()
 	{
-	}
-
-	public DownloadDescription(String RequestID, String DestinationLocation, int MaxRetryCount, int IndividualURLRetryCount, int RequestPriority, int GroupID, String[] URLs, boolean bIsPaused, long TotalBytesNeeded)
-	{
-		this.RequestID = RequestID;
-		this.DestinationLocation = DestinationLocation;
-		this.MaxRetryCount = MaxRetryCount;
-		this.IndividualURLRetryCount = IndividualURLRetryCount;
-		this.RequestPriority = RequestPriority;
-		this.GroupID = GroupID;
-		this.URLs = URLs;
-		this.bIsPaused = bIsPaused;
-		this.TotalBytesNeeded = TotalBytesNeeded;
+		URLs = new ArrayList<String>();
 	}
 	
 	public static ArrayList<DownloadDescription> BuildDescriptionArray()
@@ -239,7 +221,12 @@ public class DownloadDescription
 			if (false == jObject.isNull(URLKey))
 			{
 				org.json.JSONArray JSONURLs = jObject.getJSONArray(URLKey);
-				ReturnedDescription.URLs = IntStream.range(0, JSONURLs.length()).mapToObj(JSONURLs::optString).toArray(String[]::new);
+				
+				int NumURLs = JSONURLs.length();
+				for (int URLIndex = 0; URLIndex < NumURLs; ++URLIndex)
+				{
+					ReturnedDescription.URLs.add(JSONURLs.optString(URLIndex));
+				}
 			}
 
 			if (false == jObject.isNull(bHasCompletedKey))
@@ -275,8 +262,8 @@ public class DownloadDescription
 		try 
 		{
 			org.json.JSONArray URLArray = new org.json.JSONArray();
-			for (String URL : URLs) {
-				URLArray.put(URL);
+			for (int URLIndex = 0; URLIndex < URLs.size(); ++URLIndex) {
+				URLArray.put(URLs.get(URLIndex));
 			}
 
 			JSONObj.put(URLKey, URLArray);
@@ -331,13 +318,13 @@ public class DownloadDescription
 	//Requires the URLs lists to match exactly
 	public boolean HasExactSameURLs(DownloadDescription Other)
 	{
-		return ((this.URLs.length == Other.URLs.length) && (!this.HasNewURLs(Other)));
+		return ((this.URLs.size() == Other.URLs.size()) && (!this.HasNewURLs(Other)));
 	}
 	
 	//This only requires that our DestinationLocation has all the same URLs as the Other, but it can have more.
 	public boolean HasNewURLs(DownloadDescription Other)
 	{
-		return !List.of(this.URLs).containsAll(List.of(Other.URLs));
+		return (!this.URLs.containsAll(Other.URLs));
 	}
 	
 	public boolean HasNewDestinationLocation(DownloadDescription Other)
diff --git a/Engine/Plugins/Online/Android/AndroidFetchBackgroundDownload/Source/AndroidFetchBackgroundDownload/Java/download/datastructs/DownloadNotificationDescription.java b/Engine/Plugins/Online/Android/AndroidFetchBackgroundDownload/Source/AndroidFetchBackgroundDownload/Java/download/datastructs/DownloadNotificationDescription.java
--- a/Engine/Plugins/Online/Android/AndroidFetchBackgroundDownload/Source/AndroidFetchBackgroundDownload/Java/download/datastructs/DownloadNotificationDescription.java
+++ b/Engine/Plugins/Online/Android/AndroidFetchBackgroundDownload/Source/AndroidFetchBackgroundDownload/Java/download/datastructs/DownloadNotificationDescription.java
@@ -3,32 +3,22 @@ package com.epicgames.unreal.download.datastructs;
 
 import android.app.NotificationManager;
 import android.content.Context;
-import android.content.SharedPreferences;
+import android.util.Log;
 
 import androidx.annotation.NonNull;
 import androidx.annotation.Nullable;
 import androidx.work.Data;
+import androidx.work.WorkerParameters;
 
 import com.epicgames.unreal.LocalNotificationReceiver;
 import com.epicgames.unreal.Logger;
-
-import java.util.Map;
+import com.epicgames.unreal.download.datastructs.DownloadWorkerParameterKeys;
 
 //Helper class that stores all the needed information for a Notification in one object and handles parsing and caching defaults stored in the WorkerParameters
 public class DownloadNotificationDescription
 {
 	//Constructor that parses good defaults from our Worker InputData
-	public DownloadNotificationDescription(@NonNull Data data, @NonNull Context context, @Nullable Logger Log)
-	{
-		this(data.getKeyValueMap(), context, Log);
-	}
-
-	public DownloadNotificationDescription(@NonNull SharedPreferences preferences, @NonNull Context context, @Nullable Logger Log)
-	{
-		this(preferences.getAll(), context, Log);
-	}
-
-	public DownloadNotificationDescription(@NonNull Map<String, ?> data, @NonNull Context context, @Nullable Logger Log)
+	public DownloadNotificationDescription(@NonNull Data data,@NonNull Context context, @Nullable Logger Log)
 	{
 		//Set values that are just raw defaults that we expect to be manually overriden
 		{
@@ -38,26 +28,26 @@ public class DownloadNotificationDescription
 					
 		//Load notification channel information
 		{
-			NotificationChannelID = data.get(DownloadWorkerParameterKeys.NOTIFICATION_CHANNEL_ID_KEY) instanceof String value ? value : null;
+			NotificationChannelID = data.getString(DownloadWorkerParameterKeys.NOTIFICATION_CHANNEL_ID_KEY);
 			if (null == NotificationChannelID) 
 			{
 				NotificationChannelID = "ue-downloadworker-channel-id";
 			}
 
-			NotificationChannelName = data.get(DownloadWorkerParameterKeys.NOTIFICATION_CHANNEL_NAME_KEY) instanceof String value ? value : null;
+			NotificationChannelName = data.getString(DownloadWorkerParameterKeys.NOTIFICATION_CHANNEL_NAME_KEY);
 			if (null == NotificationChannelName) 
 			{
 				NotificationChannelName = "ue-downloadworker-channel";
 			}
 
-			NotificationChannelImportance = data.get(DownloadWorkerParameterKeys.NOTIFICATION_CHANNEL_IMPORTANCE_KEY) instanceof Integer value ? value : NotificationManager.IMPORTANCE_DEFAULT;
+			NotificationChannelImportance = data.getInt(DownloadWorkerParameterKeys.NOTIFICATION_CHANNEL_IMPORTANCE_KEY, NotificationManager.IMPORTANCE_DEFAULT);
 		}
 		
 		//Load notification base information
 		{
 			//Loads from data or defaults to a random number (that is hopefully unique) as this can NOT be set to 0 for SetForeground notifications!
 
-			NotificationID = data.get(DownloadWorkerParameterKeys.NOTIFICATION_ID_KEY) instanceof Integer value ? value : DownloadWorkerParameterKeys.NOTIFICATION_DEFAULT_ID_KEY;
+			NotificationID = data.getInt(DownloadWorkerParameterKeys.NOTIFICATION_ID_KEY, DownloadWorkerParameterKeys.NOTIFICATION_DEFAULT_ID_KEY);
 
 			if (NotificationID == 0)
 			{
@@ -70,62 +60,60 @@ public class DownloadNotificationDescription
 
 		//Configuration 
 		{
-			ShouldHandleCellular = data.get(DownloadWorkerParameterKeys.CELLULAR_HANDLE_KEY) instanceof Boolean value ? value : false;
-
-			ShouldShowPercentage = data.get(DownloadWorkerParameterKeys.SHOW_PERCENTAGE_KEY) instanceof Boolean value ? value : false;
+			ShouldHandleCellular = data.getBoolean(DownloadWorkerParameterKeys.CELLULAR_HANDLE_KEY, false);
 		}
 
 		//Load notification content information
 		{
-			TitleText = data.get(DownloadWorkerParameterKeys.NOTIFICATION_CONTENT_TITLE_KEY) instanceof String value ? value : null;
+			TitleText = data.getString(DownloadWorkerParameterKeys.NOTIFICATION_CONTENT_TITLE_KEY);
 			if (null == TitleText) 
 			{
 				TitleText = "Downloading";
 			}
 
-			ContentText = data.get(DownloadWorkerParameterKeys.NOTIFICATION_CONTENT_TEXT_KEY) instanceof String value ? value : null;
+			ContentText = data.getString(DownloadWorkerParameterKeys.NOTIFICATION_CONTENT_TEXT_KEY);
 			if (null == ContentText) 
 			{
 				ContentText = "Download in Progress";
 			}
 
-			ContentCompleteText = data.get(DownloadWorkerParameterKeys.NOTIFICATION_CONTENT_COMPLETE_TEXT_KEY) instanceof String value ? value : null;
+			ContentCompleteText = data.getString(DownloadWorkerParameterKeys.NOTIFICATION_CONTENT_COMPLETE_TEXT_KEY);
 			if (null == ContentCompleteText) 
 			{
 				ContentCompleteText = "Complete";
 			}
 			
-			CancelText = data.get(DownloadWorkerParameterKeys.NOTIFICATION_CONTENT_CANCEL_DOWNLOAD_TEXT_KEY) instanceof String value ? value : null;
+			CancelText = data.getString(DownloadWorkerParameterKeys.NOTIFICATION_CONTENT_CANCEL_DOWNLOAD_TEXT_KEY);
 			if (null == CancelText) 
 			{
 				CancelText = "Cancel";
 			}
 
-			NoInternetAvailable = data.get(DownloadWorkerParameterKeys.NOTIFICATION_CONTENT_NO_INTERNET_TEXT_KEY) instanceof String value ? value : null;
+			NoInternetAvailable = data.getString(DownloadWorkerParameterKeys.NOTIFICATION_CONTENT_NO_INTERNET_TEXT_KEY);
 			if (null == NoInternetAvailable)
 			{
 				NoInternetAvailable = "No Internet Available";
 			}
 
-			AirplaneModeText = data.get(DownloadWorkerParameterKeys.NOTIFICATION_CONTENT_AIRPLANE_MODE_TEXT_KEY) instanceof String value ? value : null;
+			AirplaneModeText = data.getString(DownloadWorkerParameterKeys.NOTIFICATION_CONTENT_AIRPLANE_MODE_TEXT_KEY);
 			if (null == AirplaneModeText)
 			{
 				AirplaneModeText = "Airplane Mode Enabled";
 			}
 			
-			DataSaverEnabledText = data.get(DownloadWorkerParameterKeys.NOTIFICATION_CONTENT_DATA_SAVER_ENABLED_TEXT_KEY) instanceof String value ? value : null;
+			DataSaverEnabledText = data.getString(DownloadWorkerParameterKeys.NOTIFICATION_CONTENT_DATA_SAVER_ENABLED_TEXT_KEY);
 			if (null == DataSaverEnabledText)
 			{
 				DataSaverEnabledText = "Data Saver Enabled";
 			}
 
-			WaitingForCellularText = data.get(DownloadWorkerParameterKeys.NOTIFICATION_CONTENT_WAITING_FOR_CELLULAR_TEXT_KEY) instanceof String value ? value : null;
+			WaitingForCellularText = data.getString(DownloadWorkerParameterKeys.NOTIFICATION_CONTENT_WAITING_FOR_CELLULAR_TEXT_KEY);
 			if (null == WaitingForCellularText)
 			{
 				WaitingForCellularText = "Waiting for cellular approval";
 			}
 
-			ApproveText = data.get(DownloadWorkerParameterKeys.NOTIFICATION_CONTENT_APPROVE_TEXT_KEY) instanceof String value ? value : null;
+			ApproveText = data.getString(DownloadWorkerParameterKeys.NOTIFICATION_CONTENT_APPROVE_TEXT_KEY);
 			if (null == ApproveText)
 			{
 				ApproveText = "Approve";
@@ -137,7 +125,7 @@ public class DownloadNotificationDescription
 			//Flag if the user tried to set anything so we can show an error if the load fails when they weren't expecting to use the default
 			boolean bTriedToSetAnything = false;
 			
-			String CancelIconResourceName = data.get(DownloadWorkerParameterKeys.NOTIFICATION_RESOURCE_CANCEL_ICON_NAME) instanceof String value ? value : null;
+			String CancelIconResourceName = data.getString(DownloadWorkerParameterKeys.NOTIFICATION_RESOURCE_CANCEL_ICON_NAME);
 			if (null == CancelIconResourceName) 
 			{
 				CancelIconResourceName = "ic_delete";
@@ -147,7 +135,7 @@ public class DownloadNotificationDescription
 				bTriedToSetAnything = true;
 			}
 
-			String CancelIconResourceType = data.get(DownloadWorkerParameterKeys.NOTIFICATION_RESOURCE_CANCEL_ICON_TYPE) instanceof String value ? value : null;
+			String CancelIconResourceType = data.getString(DownloadWorkerParameterKeys.NOTIFICATION_RESOURCE_CANCEL_ICON_TYPE);
 			if (null == CancelIconResourceType) 
 			{
 				CancelIconResourceType = "drawable";
@@ -157,7 +145,7 @@ public class DownloadNotificationDescription
 				bTriedToSetAnything = true;
 			}
 
-			String CancelIconResourcePackage = data.get(DownloadWorkerParameterKeys.NOTIFICATION_RESOURCE_CANCEL_ICON_PACKAGE) instanceof String value ? value : null;
+			String CancelIconResourcePackage = data.getString(DownloadWorkerParameterKeys.NOTIFICATION_RESOURCE_CANCEL_ICON_PACKAGE);
 			if (null == CancelIconResourcePackage) 
 			{
 				CancelIconResourcePackage = context.getPackageName();
@@ -198,7 +186,7 @@ public class DownloadNotificationDescription
 			//Flag if the user tried to set anything so we can show an error if the load fails when they weren't expecting to use the default
 			boolean bTriedToSetAnything = false;
 			
-			String SmallIconResourceName = data.get(DownloadWorkerParameterKeys.NOTIFICATION_RESOURCE_SMALL_ICON_NAME) instanceof String value ? value : null;
+			String SmallIconResourceName = data.getString(DownloadWorkerParameterKeys.NOTIFICATION_RESOURCE_SMALL_ICON_NAME);
 			if (null == SmallIconResourceName)
 			{
 				SmallIconResourceName = "ic_notification_simple";
@@ -208,7 +196,7 @@ public class DownloadNotificationDescription
 				bTriedToSetAnything = true;
 			}
 
-			String SmallIconResourceType = data.get(DownloadWorkerParameterKeys.NOTIFICATION_RESOURCE_SMALL_ICON_TYPE) instanceof String value ? value : null;
+			String SmallIconResourceType = data.getString(DownloadWorkerParameterKeys.NOTIFICATION_RESOURCE_SMALL_ICON_TYPE);
 			if (null == SmallIconResourceType)
 			{
 				SmallIconResourceType = "drawable";
@@ -218,7 +206,7 @@ public class DownloadNotificationDescription
 				bTriedToSetAnything = true;
 			}
 
-			String SmallIconResourcePackage = data.get(DownloadWorkerParameterKeys.NOTIFICATION_RESOURCE_SMALL_ICON_PACKAGE) instanceof String value ? value : null;
+			String SmallIconResourcePackage = data.getString(DownloadWorkerParameterKeys.NOTIFICATION_RESOURCE_SMALL_ICON_PACKAGE);
 			if (null == SmallIconResourcePackage)
 			{
 				SmallIconResourcePackage = context.getPackageName();
@@ -271,8 +259,6 @@ public class DownloadNotificationDescription
 
 	public boolean ShouldHandleCellular = false;
 
-	public boolean ShouldShowPercentage = false;
-
 	public String TitleText = null;
 	public String ContentText = null;
 	public String ContentCompleteText = null;
diff --git a/Engine/Plugins/Online/Android/AndroidFetchBackgroundDownload/Source/AndroidFetchBackgroundDownload/Java/download/datastructs/DownloadQueueDescription.java b/Engine/Plugins/Online/Android/AndroidFetchBackgroundDownload/Source/AndroidFetchBackgroundDownload/Java/download/datastructs/DownloadQueueDescription.java
--- a/Engine/Plugins/Online/Android/AndroidFetchBackgroundDownload/Source/AndroidFetchBackgroundDownload/Java/download/datastructs/DownloadQueueDescription.java
+++ b/Engine/Plugins/Online/Android/AndroidFetchBackgroundDownload/Source/AndroidFetchBackgroundDownload/Java/download/datastructs/DownloadQueueDescription.java
@@ -2,47 +2,48 @@
 package com.epicgames.unreal.download.datastructs;
 
 import android.content.Context;
-import android.content.SharedPreferences;
+import android.util.Log;
 
 import androidx.annotation.NonNull;
 import androidx.annotation.Nullable;
 import androidx.work.Data;
+import androidx.work.WorkerParameters;
 
 import com.epicgames.unreal.Logger;
+
+import com.epicgames.unreal.download.datastructs.DownloadDescription;
+import com.epicgames.unreal.download.datastructs.DownloadWorkerParameterKeys;
 import com.epicgames.unreal.download.DownloadProgressListener;
 
 import org.json.JSONArray;
 import org.json.JSONException;
 
 import java.util.ArrayList;
-import java.util.Map;
+import java.util.List;
 
 
 //Helper class that is used to serialize any information needed for our downloads as a whole from/to WorkerParameters
 public class DownloadQueueDescription
 {
-	public final static int DEFAULT_MAX_CONCURRENT_DOWNLOADS = 4;
-
-	public DownloadQueueDescription(ArrayList<DownloadDescription> downloadDescriptions, @NonNull Map<String, ?> data)
-	{
-		DownloadDescriptions = downloadDescriptions;
-		MaxConcurrentDownloads = data.get(DownloadWorkerParameterKeys.DOWNLOAD_MAX_CONCURRENT_REQUESTS_KEY) instanceof Integer value ? value : DEFAULT_MAX_CONCURRENT_DOWNLOADS;
-	}
-
-	public DownloadQueueDescription(ArrayList<DownloadDescription> downloadDescriptions, @NonNull SharedPreferences backgroundPreferences)
-	{
-		this(downloadDescriptions, backgroundPreferences.getAll());
-	}
-
 	//Constructor that parses the List<DownloadDescription> from our Worker InputData (WorkerParameters)
-	public DownloadQueueDescription(@NonNull String DownloadDescriptionListFileName, @NonNull Data data, @Nullable Logger Log)
-	{
-		this(ParseDownloadDescriptionString(DownloadDescriptionListFileName, Log), data.getKeyValueMap());
+	public DownloadQueueDescription(@NonNull Data data, @NonNull Context context, @Nullable Logger Log)
+	{		
+		//Parse DownloadDescriptions 
+		String DownloadDescriptionListFileName = GetDownloadDescriptionListFileName(data, Log);
+		if (DownloadDescriptionListFileName != null)
+		{
+			ParseDownloadDescriptionString(DownloadDescriptionListFileName, Log);
+		}
+		
+		MaxConcurrentDownloads = data.getInt(DownloadWorkerParameterKeys.DOWNLOAD_MAX_CONCURRENT_REQUESTS_KEY, DEFAULT_MAX_CONCURRENT_DOWNLOADS);
 	}
 
-	public DownloadQueueDescription(@NonNull String DownloadDescriptionListFileName, @NonNull Data data, @NonNull Context context, @Nullable Logger Log)
+	public DownloadQueueDescription(String DownloadDescriptionListFileName, @NonNull Data data, @NonNull Context context, @Nullable Logger Log)
 	{
-		this(DownloadDescriptionListFileName, data, Log);
+		//Parse DownloadDescriptions 
+		ParseDownloadDescriptionString(DownloadDescriptionListFileName, Log);
+
+		MaxConcurrentDownloads = data.getInt(DownloadWorkerParameterKeys.DOWNLOAD_MAX_CONCURRENT_REQUESTS_KEY, DEFAULT_MAX_CONCURRENT_DOWNLOADS);
 	}
 
 	public static String GetDownloadDescriptionListFileName(@NonNull Data data, @Nullable Logger Log)
@@ -61,9 +62,9 @@ public class DownloadQueueDescription
 	}
 
 	//Parses our DownloadDescriptions member based on the passed in JSONObject file
-	static ArrayList<DownloadDescription> ParseDownloadDescriptionString(@NonNull String DownloadDescriptionListStringFileName, @Nullable Logger Log)
+	private void ParseDownloadDescriptionString(@NonNull String DownloadDescriptionListStringFileName, @Nullable Logger Log)
 	{
-		ArrayList<DownloadDescription> DownloadDescriptions = new ArrayList<>();
+		DownloadDescriptions = new ArrayList<DownloadDescription>();
 		
 		try 
 		{
@@ -94,12 +95,12 @@ public class DownloadQueueDescription
 				Log.error("Exception while parsing download descriptions: " + e.getMessage());
 			}			
 		}
-		
-		return DownloadDescriptions;
 	}
 
-	public ArrayList<DownloadDescription> DownloadDescriptions;
+	public ArrayList<DownloadDescription> DownloadDescriptions = new ArrayList<DownloadDescription>();
 	public DownloadProgressListener ProgressListener = null;
 	public int DownloadGroupID = 0;
-	public int MaxConcurrentDownloads = DEFAULT_MAX_CONCURRENT_DOWNLOADS;
+	public int MaxConcurrentDownloads = 4;
+	
+	private int DEFAULT_MAX_CONCURRENT_DOWNLOADS = 4;
 }
\ No newline at end of file
diff --git a/Engine/Plugins/Online/Android/AndroidFetchBackgroundDownload/Source/AndroidFetchBackgroundDownload/Java/download/datastructs/DownloadWorkerParameterKeys.java b/Engine/Plugins/Online/Android/AndroidFetchBackgroundDownload/Source/AndroidFetchBackgroundDownload/Java/download/datastructs/DownloadWorkerParameterKeys.java
--- a/Engine/Plugins/Online/Android/AndroidFetchBackgroundDownload/Source/AndroidFetchBackgroundDownload/Java/download/datastructs/DownloadWorkerParameterKeys.java
+++ b/Engine/Plugins/Online/Android/AndroidFetchBackgroundDownload/Source/AndroidFetchBackgroundDownload/Java/download/datastructs/DownloadWorkerParameterKeys.java
@@ -31,8 +31,6 @@ public class DownloadWorkerParameterKeys
 	public final static int NOTIFICATION_DEFAULT_ID_KEY = 1923901283;
 
 	public final static String CELLULAR_HANDLE_KEY = "CellularHandleKey";
-
-	public final static String SHOW_PERCENTAGE_KEY = "bShowPercentage";
 	
 	public final static String NOTIFICATION_CONTENT_TITLE_KEY = "NotificationContentTitle";
 	public final static String NOTIFICATION_CONTENT_TEXT_KEY = "NotificationContentText";
diff --git a/Engine/Plugins/Online/Android/AndroidFetchBackgroundDownload/Source/AndroidFetchBackgroundDownload/Java/download/fetch/FetchManagerInterfaceHelper.java b/Engine/Plugins/Online/Android/AndroidFetchBackgroundDownload/Source/AndroidFetchBackgroundDownload/Java/download/fetch/FetchManagerInterfaceHelper.java
--- a/Engine/Plugins/Online/Android/AndroidFetchBackgroundDownload/Source/AndroidFetchBackgroundDownload/Java/download/fetch/FetchManagerInterfaceHelper.java
+++ b/Engine/Plugins/Online/Android/AndroidFetchBackgroundDownload/Source/AndroidFetchBackgroundDownload/Java/download/fetch/FetchManagerInterfaceHelper.java
@@ -31,7 +31,6 @@ public class FetchManagerInterfaceHelper {
 					int deletedCount = 0;
 					for (File file : files) {
 						if (!file.getName().endsWith("DownloadDescriptionJSONs") &&
-							!file.getName().endsWith("DownloadDescriptions") &&
 							!file.getName().endsWith("URLMap") &&
 							!file.getName().equals(GameActivity.BGDL_BUILD_VERSION)) {
 							Log.debug("cleaning up: " + file.getName());
diff --git a/Engine/Plugins/Online/Android/AndroidFetchBackgroundDownload/Source/AndroidFetchBackgroundDownload/Java/download/fetch/FetchManagerNew.java b/Engine/Plugins/Online/Android/AndroidFetchBackgroundDownload/Source/AndroidFetchBackgroundDownload/Java/download/fetch/FetchManagerNew.java
--- a/Engine/Plugins/Online/Android/AndroidFetchBackgroundDownload/Source/AndroidFetchBackgroundDownload/Java/download/fetch/FetchManagerNew.java
+++ b/Engine/Plugins/Online/Android/AndroidFetchBackgroundDownload/Source/AndroidFetchBackgroundDownload/Java/download/fetch/FetchManagerNew.java
@@ -4,22 +4,31 @@ package com.epicgames.unreal.download.fetch;
 
 import android.content.Context;
 import android.content.SharedPreferences;
-
-import androidx.annotation.NonNull;
+import android.net.Uri;
 
 import com.epicgames.unreal.Logger;
+import com.epicgames.unreal.download.datastructs.DownloadDescription;
+import com.epicgames.unreal.download.datastructs.DownloadQueueDescription;
 import com.epicgames.unreal.download.DownloadProgressListener;
-import com.epicgames.unreal.download.UEDownloadWorker.EDownloadCompleteReason;
-import com.epicgames.unreal.download.asyncdownloader.AsyncDownloadRequest;
-import com.epicgames.unreal.download.asyncdownloader.AsyncDownloadRequestListener;
 import com.epicgames.unreal.download.asyncdownloader.AsyncDownloader;
 import com.epicgames.unreal.download.asyncdownloader.AsyncDownloaderListener;
-import com.epicgames.unreal.download.datastructs.DownloadDescription;
-import com.epicgames.unreal.download.datastructs.DownloadQueueDescription;
+import com.epicgames.unreal.download.asyncdownloader.AsyncDownloadRequest;
+import com.epicgames.unreal.download.asyncdownloader.AsyncDownloadRequestListener;
+
+import com.epicgames.unreal.download.UEDownloadWorker.EDownloadCompleteReason;
+
+import okhttp3.OkHttpClient;
+
+import androidx.annotation.Nullable;
+import androidx.annotation.NonNull;
 
 import java.io.File;
-import java.net.URL;
+import java.util.ArrayList;
+import java.util.HashMap;
 import java.util.List;
+import java.util.Set;
+import java.util.Map;
+import java.net.URL;
 
 public class FetchManagerNew implements FetchManagerInterface
 {
@@ -31,11 +40,6 @@ public class FetchManagerNew implements FetchManagerInterface
 	{
 		return SharedManager;
 	}
-	
-	public Object SyncRoot()
-	{
-		return AsyncDownloader;
-	}
 
 	@Override	
 	public synchronized void StopWork(String WorkID)
@@ -46,81 +50,20 @@ public class FetchManagerNew implements FetchManagerInterface
 			AsyncDownloader.CancelAll();
 		}
 	}
-
-	public synchronized void Enqueue(@NonNull Context context, @NonNull DownloadQueueDescription QueueDescription)
-	{
-		if (AsyncDownloader == null)
-		{
-			SharedPreferences preferences = context.getSharedPreferences("BackgroundDownload", Context.MODE_PRIVATE);
-			boolean AllowHTTP2 = preferences.getBoolean("AllowHTTP2", false);
-			AsyncDownloader = new AsyncDownloader(AllowHTTP2, QueueDescription.MaxConcurrentDownloads);
-			SetupAllDownloadsCompletedListener_Internal();
-		}
-		else
-		{
-			AsyncDownloader.SetMaxConcurrent(QueueDescription.MaxConcurrentDownloads);
-		}
-
-		GlobalDownloadProgressListener = QueueDescription.ProgressListener;
-
-		for (DownloadDescription DownloadDescription : QueueDescription.DownloadDescriptions)
-		{
-			//noinspection SynchronizeOnNonFinalField
-			synchronized (AsyncDownloader)
-			{
-				final String TargetLocation = DownloadDescription.DestinationLocation;
-				final String DownloadLocation = DownloadDescription.DestinationLocation + ".temp";
-				final boolean bFileExists = new File(TargetLocation).exists();
-
-				if (bFileExists)
-				{
-					GlobalDownloadProgressListener.OnDownloadMetrics(DownloadDescription.DestinationLocation, new File(TargetLocation).length(), 0, DownloadDescription.DownloadStartTime, DownloadDescription.DownloadStartTime);
-					GlobalDownloadProgressListener.OnDownloadComplete(DownloadDescription.DestinationLocation, EDownloadCompleteReason.Success);
-
-					AsyncDownloader.Cancel(DownloadDescription.DestinationLocation);
-					continue;
-				}
-
-				final boolean bSuccess = AsyncDownloader.Submit(
-					new AsyncDownloadRequest.Builder()
-						.SetId(DownloadDescription.DestinationLocation)
-						.SetRequestId(DownloadDescription.RequestID)
-						.SetUrls(DownloadDescription.URLs)
-						.SetMaxRetries(DownloadDescription.MaxRetryCount)
-						.SetGroupID(DownloadDescription.GroupID)
-						.SetTargetFile(new File(DownloadLocation))
-						.SetResumeIfPossible(true)
-						.SetTotalSize(DownloadDescription.TotalBytesNeeded)
-						.Build(),
-					new AsyncDownloadRequestListenerImpl(this, GlobalDownloadProgressListener, TargetLocation, DownloadLocation));
-				if (!bSuccess)
-				{
-					RetryDownload(DownloadDescription.DestinationLocation);
-				}
-
-				GlobalDownloadProgressListener.OnDownloadEnqueued(DownloadDescription);
-			}
-		}
-	}
 	
 	@Override
 	public synchronized void EnqueueRequests(@NonNull Context context, @NonNull DownloadQueueDescription QueueDescription)
 	{
 		Log.debug("FetchManagerNew: EnqueueRequests with " + QueueDescription.DownloadDescriptions.size() + " requests");
 
-		if (AsyncDownloader != null)
-		{
-			AsyncDownloader.CancelAll();
-		}
-
 		NumFailedRequests = 0;
-		if (QueueDescription.DownloadDescriptions.isEmpty()) {
+		if (QueueDescription.DownloadDescriptions.size() == 0) {
 			return;
 		}
 
 		if (AsyncDownloader == null)
 		{
-			SharedPreferences preferences = context.getSharedPreferences("BackgroundDownload", Context.MODE_PRIVATE);
+			SharedPreferences preferences = context.getSharedPreferences("BackgroundDownload", context.MODE_PRIVATE);
 			boolean AllowHTTP2 = preferences.getBoolean("AllowHTTP2", false);
 			AsyncDownloader = new AsyncDownloader(AllowHTTP2, QueueDescription.MaxConcurrentDownloads);
 			SetupAllDownloadsCompletedListener_Internal();
@@ -134,41 +77,22 @@ public class FetchManagerNew implements FetchManagerInterface
 
 		for (DownloadDescription DownloadDescription : QueueDescription.DownloadDescriptions)
 		{
-			//noinspection SynchronizeOnNonFinalField
-			synchronized (AsyncDownloader)
-			{
-				final String TargetLocation = DownloadDescription.DestinationLocation;
-				final String DownloadLocation = DownloadDescription.DestinationLocation + ".temp";
-				final boolean bFileExists = new File(TargetLocation).exists();
-
-				if (bFileExists)
-				{
-					GlobalDownloadProgressListener.OnDownloadMetrics(DownloadDescription.DestinationLocation, new File(TargetLocation).length(), 0, DownloadDescription.DownloadStartTime, DownloadDescription.DownloadStartTime);
-					GlobalDownloadProgressListener.OnDownloadComplete(DownloadDescription.DestinationLocation, EDownloadCompleteReason.Success);
-					
-					AsyncDownloader.Cancel(DownloadDescription.DestinationLocation);
-					continue;
-				}
-				
-				final boolean bSuccess = AsyncDownloader.Submit(
-					new AsyncDownloadRequest.Builder()
-						.SetId(DownloadDescription.DestinationLocation)
-						.SetRequestId(DownloadDescription.RequestID)
-						.SetUrls(DownloadDescription.URLs)
-						.SetMaxRetries(DownloadDescription.MaxRetryCount)
-						.SetGroupID(DownloadDescription.GroupID)
-						.SetTargetFile(new File(DownloadLocation))
-						.SetResumeIfPossible(true)
-						.SetTotalSize(DownloadDescription.TotalBytesNeeded)
-						.Build(),
-					new AsyncDownloadRequestListenerImpl(this, GlobalDownloadProgressListener, TargetLocation, DownloadLocation));
-				if (!bSuccess)
-				{
-					RetryDownload(DownloadDescription.DestinationLocation);
-				}
-
-				GlobalDownloadProgressListener.OnDownloadEnqueued(DownloadDescription);
-			}
+			final String TargetLocation = DownloadDescription.DestinationLocation;
+			final String DownloadLocation = DownloadDescription.DestinationLocation + ".temp";
+
+			final boolean bSuccess = AsyncDownloader.Submit(
+				new AsyncDownloadRequest.Builder()
+					.SetId(DownloadDescription.DestinationLocation)
+					.SetRequestId(DownloadDescription.RequestID)
+					.SetUrls(List.copyOf(DownloadDescription.URLs))
+					.SetMaxRetries(DownloadDescription.MaxRetryCount)
+					.SetGroupID(DownloadDescription.GroupID)
+					.SetTargetFile(new File(DownloadLocation))
+					.SetResumeIfPossible(true)
+					.SetTotalSize(DownloadDescription.TotalBytesNeeded)
+					.Build(),
+				new AsyncDownloadRequestListenerImpl(this, GlobalDownloadProgressListener, TargetLocation, DownloadLocation));
+			GlobalDownloadProgressListener.OnDownloadEnqueued(DownloadDescription.DestinationLocation, bSuccess);
 		}
 
 		// clear flag that stops downloader from downloading
@@ -182,9 +106,9 @@ public class FetchManagerNew implements FetchManagerInterface
 			return;
 		}
 		if (bPause) {
-			AsyncDownloader.Pause(DestinationLocation, true);
+			AsyncDownloader.Pause(DestinationLocation);
 		} else {
-			AsyncDownloader.Resume(DestinationLocation, true);
+			AsyncDownloader.Resume(DestinationLocation);
 		}
 	}
 
@@ -294,20 +218,21 @@ public class FetchManagerNew implements FetchManagerInterface
 		}
 	}
 	
-	private static final FetchManagerNew SharedManager = new FetchManagerNew();
+	private static volatile FetchManagerNew SharedManager = new FetchManagerNew();
+	private volatile OkHttpClient OkHttpClient = null;
 	private volatile AllDownloadsCompletedListener AllDownloadsCompletedListener = null;
 	private volatile AsyncDownloader AsyncDownloader = null;
 	private volatile DownloadProgressListener GlobalDownloadProgressListener = null;
 	private volatile int NumFailedRequests = 0;
-	private final Logger Log = new Logger("UE", "FetchManagerNew");
+	private Logger Log = new Logger("UE", "FetchManagerNew");
 
 	// ---------------------------------------------------------------------------------------------------------------------
 	// private helper class
 	// ---------------------------------------------------------------------------------------------------------------------
 
-	static final private class AsyncDownloadRequestListenerImpl implements AsyncDownloadRequestListener {
+	final private class AsyncDownloadRequestListenerImpl implements AsyncDownloadRequestListener {
 
-		public AsyncDownloadRequestListenerImpl(@NonNull FetchManagerNew Owner, @NonNull DownloadProgressListener ProgressListener, @NonNull String TargetLocation, @NonNull String DownloadLocation) {
+		public AsyncDownloadRequestListenerImpl(@NonNull FetchManagerNew Owner, @NonNull DownloadProgressListener ProgressListener, String TargetLocation, String DownloadLocation) {
 			this.Owner = Owner;
 			this.ProgressListener = ProgressListener;
 			this.TargetLocation = TargetLocation;
@@ -315,24 +240,28 @@ public class FetchManagerNew implements FetchManagerInterface
 		}
 
 		@Override
-		public void OnStart(@NonNull String Id) {
-			StartTime = System.currentTimeMillis();
+		public void OnStart(String Id, URL url, long startingOffset, long totalBytesOrMinus1) {
+			LastDownloadedSoFar = startingOffset;
+			TotalBytes = totalBytesOrMinus1;			// total target file size, regardless of starting offset
+
+			if (StartTime == -1L) {
+				StartTime = System.currentTimeMillis();
+			}
+
+			// reporting LastDownloadedSoFar as "bytes downloaded since last call" gives less jumpy progress bars
+			ProgressListener.OnDownloadProgress(Id, LastDownloadedSoFar, LastDownloadedSoFar);	
 		}
 		@Override
-		public void OnProgress(@NonNull String Id, long bytesDownloadedSoFar) {
+		public void OnProgress(String Id, long bytesDownloadedSoFar, long totalBytesOrMinus1) {
 			// Log.debug("OnProgress: " + Id + " | " + bytesDownloadedSoFar + " | " + totalBytesOrMinus1);
-			LastDownloadedSoFar = bytesDownloadedSoFar;
 
-			long NewUpdateTime = System.currentTimeMillis();
-			if (NewUpdateTime >= UpdateTime + 100) {
-				long TotalDownloadedSinceLastUpdate = LastDownloadedSoFar - LastUpdatedSoFar;
-				LastUpdatedSoFar = LastDownloadedSoFar;
-				UpdateTime = NewUpdateTime;
-				ProgressListener.OnDownloadProgress(Id, LastUpdatedSoFar, TotalDownloadedSinceLastUpdate);
-			}
+			// the Math.max below is only necessary if resuming a download after an error, and the server suddenly stops supporting "206 Partial Content"
+			long TotalDownloadedSinceLastCall = Math.max(0L, bytesDownloadedSoFar - LastDownloadedSoFar);	
+			LastDownloadedSoFar = bytesDownloadedSoFar;
+			ProgressListener.OnDownloadProgress(Id, TotalDownloadedSinceLastCall, LastDownloadedSoFar);
 		}
 		@Override
-		public void OnSuccess(@NonNull String Id, @NonNull File file) {
+		public void OnSuccess(String Id, File file) {
 			EDownloadCompleteReason Reason = EDownloadCompleteReason.Success;
 			File NewFileAtDestinationLocation = new File(TargetLocation);
 			if (!file.renameTo(NewFileAtDestinationLocation)) {
@@ -343,16 +272,16 @@ public class FetchManagerNew implements FetchManagerInterface
 			OnCompleteDownload_Internal(Id, Reason);
 		}
 		@Override
-		public void OnError(@NonNull String Id, @NonNull Throwable Error, URL LastTriedUrl) {
+		public void OnError(String Id, Throwable Error, @Nullable URL LastTriedUrl) {
 			Log.debug("OnError: " + Id);
 			OnCompleteDownload_Internal(Id, EDownloadCompleteReason.OutOfRetries);
 		}
 		@Override
-		public void OnPaused(@NonNull String Id) {
+		public void OnPaused(String Id) {
 			// do we need to do something here?
 		}
 
-		private void OnCompleteDownload_Internal(@NonNull String Id, EDownloadCompleteReason Reason) {
+		private void OnCompleteDownload_Internal(String Id, EDownloadCompleteReason Reason) {
 			long EndTime = System.currentTimeMillis();
 
 			long Duration = 0;
@@ -364,24 +293,18 @@ public class FetchManagerNew implements FetchManagerInterface
 				Owner.NumFailedRequests++;
 			}
 
-			long TotalDownloadedSinceLastUpdate = LastDownloadedSoFar - LastUpdatedSoFar;
-			if (TotalDownloadedSinceLastUpdate > 0)	{
-				LastUpdatedSoFar = LastDownloadedSoFar;
-				UpdateTime = EndTime;
-				ProgressListener.OnDownloadProgress(Id, LastUpdatedSoFar, TotalDownloadedSinceLastUpdate);
-			}
-			ProgressListener.OnDownloadMetrics(Id, LastUpdatedSoFar, Duration, StartTime, EndTime);
+			ProgressListener.OnDownloadMetrics(Id, TotalBytes, Duration, StartTime, EndTime);
 			ProgressListener.OnDownloadComplete(Id, Reason);
 		}
 
 		private long LastDownloadedSoFar = 0;
-		private long StartTime = -1;
-		private long LastUpdatedSoFar = 0;
-		private long UpdateTime = Long.MIN_VALUE;
-		private final @NonNull DownloadProgressListener ProgressListener;
-		private final @NonNull FetchManagerNew Owner;
-		private final Logger Log = new Logger("UE", "FetchManagerNew");
-		private final String TargetLocation;
-		private final String DownloadLocation;
+		private long TotalBytes = 0;
+		private long StartTime = -1L;
+		private @NonNull DownloadProgressListener ProgressListener = null;
+		private @NonNull FetchManagerNew Owner;
+		private Logger Log = new Logger("UE", "FetchManagerNew");
+		private String TargetLocation;
+		private String DownloadLocation;
 	}
+
 }
diff --git a/Engine/Plugins/Online/Android/AndroidFetchBackgroundDownload/Source/AndroidFetchBackgroundDownload/Java/download/fetch/FetchManagerOld.java b/Engine/Plugins/Online/Android/AndroidFetchBackgroundDownload/Source/AndroidFetchBackgroundDownload/Java/download/fetch/FetchManagerOld.java
--- a/Engine/Plugins/Online/Android/AndroidFetchBackgroundDownload/Source/AndroidFetchBackgroundDownload/Java/download/fetch/FetchManagerOld.java
+++ b/Engine/Plugins/Online/Android/AndroidFetchBackgroundDownload/Source/AndroidFetchBackgroundDownload/Java/download/fetch/FetchManagerOld.java
@@ -160,9 +160,21 @@ public class FetchManagerOld implements FetchDownloadProgressOwner, FetchEnqueue
 			
 			if (ActiveDescription == null)
 			{
-				//Don't even have an entry in the old list, so it's completely new
-				QueueNewDownloadDescription(NewDownloadDescription, LastRequest);
-				continue;
+				if (NewDownloadDescription.RequestID == null || NewDownloadDescription.RequestID.isEmpty())
+				{
+					//Don't even have an entry in the old list, so it's completely new
+					QueueNewDownloadDescription(NewDownloadDescription, LastRequest);
+					continue;
+				}
+				
+				ActiveDescription = TempRequestedDownloadsCopy.values().stream().filter(Description -> NewDownloadDescription.RequestID.equals(Description.RequestID)).findFirst().orElse(null);
+				
+				if (ActiveDescription == null)
+				{
+					//Don't even have an entry for this RequestID in the old list, so it's completely new
+					QueueNewDownloadDescription(NewDownloadDescription, LastRequest);
+					continue;
+				}
 			}
 			
 			//Need to update our DownloadDescription
@@ -467,10 +479,10 @@ public class FetchManagerOld implements FetchDownloadProgressOwner, FetchEnqueue
 			return null;
 		}
 		
-		final int NumURLs = Description.URLs.length;
+		final int NumURLs = Description.URLs.size();
 		final int URLIndexToUse = Description.CurrentRetryCount % NumURLs;
 		
-		return Description.URLs[URLIndexToUse];
+		return Description.URLs.get(URLIndexToUse);
 	}
 	
 	private boolean IsOutOfRetries(DownloadDescription Description)
@@ -883,7 +895,7 @@ public class FetchManagerOld implements FetchDownloadProgressOwner, FetchEnqueue
 			return;
 		}
 
-		MatchingDescription.ProgressListener.OnDownloadEnqueued(MatchingDescription);
+		MatchingDescription.ProgressListener.OnDownloadEnqueued(DestinationLocation, true);
 	}
 
 	@Override
@@ -917,9 +929,7 @@ public class FetchManagerOld implements FetchDownloadProgressOwner, FetchEnqueue
 			return;
 		}
 
-		RetryDownload(DestinationLocation);
-
-		MatchingDescription.ProgressListener.OnDownloadEnqueued(MatchingDescription);
+		MatchingDescription.ProgressListener.OnDownloadEnqueued(DestinationLocation, false);
 	}
 	
 	public synchronized void HandleEnqueueCheckCallback(String DestinationLocation, @Nullable Download FetchDownload)
@@ -969,12 +979,7 @@ public class FetchManagerOld implements FetchDownloadProgressOwner, FetchEnqueue
 			MatchingDescription.CachedFetchID = FetchDownload.getId();
 		}
 
-		if (!bIsAlreadyEnqueued)
-		{
-			RetryDownload(DestinationLocation);
-		}
-
-		MatchingDescription.ProgressListener.OnDownloadEnqueued(MatchingDescription);
+		MatchingDescription.ProgressListener.OnDownloadEnqueued(DestinationLocation, bIsAlreadyEnqueued);
 	}
 
 	public boolean IsErrorPossiblyIndicatingPreviouslyQueued(@NonNull Error EnqueueError)
diff --git a/Engine/Plugins/Online/Android/AndroidFetchBackgroundDownload/Source/AndroidFetchBackgroundDownload/Private/AndroidDownloadManager.cpp b/Engine/Plugins/Online/Android/AndroidFetchBackgroundDownload/Source/AndroidFetchBackgroundDownload/Private/AndroidDownloadManager.cpp
deleted file mode 100644
diff --git a/Engine/Plugins/Online/Android/AndroidFetchBackgroundDownload/Source/AndroidFetchBackgroundDownload/Private/AndroidFetchPlatformBackgroundHttp.cpp b/Engine/Plugins/Online/Android/AndroidFetchBackgroundDownload/Source/AndroidFetchBackgroundDownload/Private/AndroidFetchPlatformBackgroundHttp.cpp
--- a/Engine/Plugins/Online/Android/AndroidFetchBackgroundDownload/Source/AndroidFetchBackgroundDownload/Private/AndroidFetchPlatformBackgroundHttp.cpp
+++ b/Engine/Plugins/Online/Android/AndroidFetchBackgroundDownload/Source/AndroidFetchBackgroundDownload/Private/AndroidFetchPlatformBackgroundHttp.cpp
@@ -2,21 +2,10 @@
 
 #include "AndroidFetchPlatformBackgroundHttp.h"
 
-#include "AndroidDownloadManager.h"
 #include "AndroidPlatformBackgroundHttpManager.h"
 #include "AndroidPlatformBackgroundHttpRequest.h"
 #include "AndroidPlatformBackgroundHttpResponse.h"
 
-#if USE_ANDROID_JNI
-extern bool AndroidThunkCpp_GetSharedPreferenceBool(const FString& Group, const FString& Key, bool DefaultValue);
-#endif
-
-static bool UseNewDownloadManager()
-{
-	static bool Value = AndroidThunkCpp_GetSharedPreferenceBool("BackgroundPreferences", "UseNewDownloadManager", false);
-	return Value;
-}
-
 void FAndroidFetchPlatformBackgroundHttp::Initialize()
 {
 }
@@ -27,31 +16,16 @@ void FAndroidFetchPlatformBackgroundHttp::Shutdown()
 
 FBackgroundHttpManagerPtr FAndroidFetchPlatformBackgroundHttp::CreatePlatformBackgroundHttpManager()
 {
-	if (UseNewDownloadManager())
-	{
-		return MakeShared<UE::Online::Download::Android::FDownloadManager, ESPMode::ThreadSafe>();
-	}
-
-	return MakeShared<FAndroidPlatformBackgroundHttpManager, ESPMode::ThreadSafe>();
+    return MakeShared<FAndroidPlatformBackgroundHttpManager, ESPMode::ThreadSafe>();
 }
 
 FBackgroundHttpRequestPtr FAndroidFetchPlatformBackgroundHttp::ConstructBackgroundRequest()
 {
-	if (UseNewDownloadManager())
-	{
-		return MakeShared<UE::Online::Download::Android::FDownloadManager::FRequest, ESPMode::ThreadSafe>();
-	}
-	
 	return MakeShared<FAndroidPlatformBackgroundHttpRequest, ESPMode::ThreadSafe>();
 }
 
 FBackgroundHttpResponsePtr FAndroidFetchPlatformBackgroundHttp::ConstructBackgroundResponse(int32 ResponseCode, const FString& TempFilePath)
 {
-	if (UseNewDownloadManager())
-	{
-		return MakeShared<FBackgroundHttpResponseImpl, ESPMode::ThreadSafe>(TempFilePath, ResponseCode);
-	}
-	
 	return MakeShared<FAndroidPlatformBackgroundHttpResponse, ESPMode::ThreadSafe>(ResponseCode, TempFilePath);
 }
 
diff --git a/Engine/Plugins/Online/Android/AndroidFetchBackgroundDownload/Source/AndroidFetchBackgroundDownload/Public/AndroidDownloadManager.h b/Engine/Plugins/Online/Android/AndroidFetchBackgroundDownload/Source/AndroidFetchBackgroundDownload/Public/AndroidDownloadManager.h
deleted file mode 100644
diff --git a/Engine/Plugins/Online/Android/AndroidFetchBackgroundDownload/Source/AndroidFetchBackgroundDownload/Public/AndroidPlatformBackgroundHttpManager.h b/Engine/Plugins/Online/Android/AndroidFetchBackgroundDownload/Source/AndroidFetchBackgroundDownload/Public/AndroidPlatformBackgroundHttpManager.h
--- a/Engine/Plugins/Online/Android/AndroidFetchBackgroundDownload/Source/AndroidFetchBackgroundDownload/Public/AndroidPlatformBackgroundHttpManager.h
+++ b/Engine/Plugins/Online/Android/AndroidFetchBackgroundDownload/Source/AndroidFetchBackgroundDownload/Public/AndroidPlatformBackgroundHttpManager.h
@@ -23,11 +23,6 @@
 //Forward Declarations
 struct FPolyglotTextData;
 
-namespace UE::Online::Download::Android
-{
-	struct FDownloadManager;
-}
-
 /**
  * Manages Background Http request that are currently being processed if we are on an Android Platform
  */
@@ -228,8 +223,6 @@ protected:
 	//We only want to respond to certain broadcast java-side work if we have scheduled BGWork
 	//otherwise it could be stale work from a previous executables run that would look like an error state (Requests we haven't queued yet, etc)
 	static volatile int32 bHasManagerScheduledBGWork;
-	
-	friend UE::Online::Download::Android::FDownloadManager;
 };
 
 //WARNING: These values MUST stay in sync with their values in DownloadWorkerParameterKeys.java!
diff --git a/Engine/Source/Runtime/Online/BackgroundHTTP/Private/BackgroundHttpResponseImpl.cpp b/Engine/Source/Runtime/Online/BackgroundHTTP/Private/BackgroundHttpResponseImpl.cpp
--- a/Engine/Source/Runtime/Online/BackgroundHTTP/Private/BackgroundHttpResponseImpl.cpp
+++ b/Engine/Source/Runtime/Online/BackgroundHTTP/Private/BackgroundHttpResponseImpl.cpp
@@ -8,13 +8,8 @@
 #include "Misc/Paths.h"
 
 FBackgroundHttpResponseImpl::FBackgroundHttpResponseImpl()
-	: ResponseCode(EHttpResponseCodes::Unknown)
-{
-}
-
-FBackgroundHttpResponseImpl::FBackgroundHttpResponseImpl(FString InTempContentFilePath, int32 InResponseCode)
-	: TempContentFilePath{MoveTemp(InTempContentFilePath)}
-	, ResponseCode{InResponseCode}
+	: TempContentFilePath()
+	, ResponseCode(EHttpResponseCodes::Unknown)
 {
 }
 
diff --git a/Engine/Source/Runtime/Online/BackgroundHTTP/Public/BackgroundHttpResponseImpl.h b/Engine/Source/Runtime/Online/BackgroundHTTP/Public/BackgroundHttpResponseImpl.h
--- a/Engine/Source/Runtime/Online/BackgroundHTTP/Public/BackgroundHttpResponseImpl.h
+++ b/Engine/Source/Runtime/Online/BackgroundHTTP/Public/BackgroundHttpResponseImpl.h
@@ -14,12 +14,11 @@ class FBackgroundHttpResponseImpl
 {
 public:
 	FBackgroundHttpResponseImpl();
-	FBackgroundHttpResponseImpl(FString InTempContentFilePath, int32 InResponseCode);
-	virtual ~FBackgroundHttpResponseImpl() override = default;
+	virtual ~FBackgroundHttpResponseImpl() = default;
 
 	//IHttpBackgroundResponse
-	virtual int32 GetResponseCode() const override;
-	virtual const FString& GetTempContentFilePath() const override;
+	virtual int32 GetResponseCode() const;
+	virtual const FString& GetTempContentFilePath() const;
 
 protected:	
 	FString TempContentFilePath;
