Add support for configrules to force software WebView compositing
Update configrules

#android

diff --git a/Engine/Build/Android/Java/src/com/epicgames/unreal/WebViewControl.java b/Engine/Build/Android/Java/src/com/epicgames/unreal/WebViewControl.java
--- a/Engine/Build/Android/Java/src/com/epicgames/unreal/WebViewControl.java
+++ b/Engine/Build/Android/Java/src/com/epicgames/unreal/WebViewControl.java
@@ -149,6 +149,13 @@ class WebViewControl
 					webView.setBackgroundColor(Color.TRANSPARENT);
 				}
 
+				// if needed for compatiblity we can set this configrules variable
+				boolean bUseSoftwareLayer = "true".equals(GameActivity.Get().ConfigRulesVars.getOrDefault("Android.ForceSoftwareWebView", "false"));
+				if (bUseSoftwareLayer)
+				{
+					webView.setLayerType(View.LAYER_TYPE_SOFTWARE, null);
+				}
+
 				// Wrap the webview in a layout that will do absolute positioning for us
 				positionLayout = new WebViewPositionLayout(GameActivity.Get(), w);
 				ViewGroup.LayoutParams params = new ViewGroup.LayoutParams(ViewGroup.LayoutParams.MATCH_PARENT, ViewGroup.LayoutParams.MATCH_PARENT);
