[Android] Add Audio focus tracking support to AndroidPlatformMisc. (also declared in GenericPlatformMisc but not implemented for other platforms yet)
Only trigger the native focus change event if there are listeners.

#rnx

diff --git a/Engine/Build/Android/Java/src/com/epicgames/unreal/GameActivity.java.template b/Engine/Build/Android/Java/src/com/epicgames/unreal/GameActivity.java.template
--- a/Engine/Build/Android/Java/src/com/epicgames/unreal/GameActivity.java.template
+++ b/Engine/Build/Android/Java/src/com/epicgames/unreal/GameActivity.java.template
@@ -171,6 +171,8 @@ import android.widget.LinearLayout;
 import android.widget.PopupWindow;
 import android.text.TextUtils;
 
+import android.media.AudioAttributes;
+import android.media.AudioFocusRequest;
 import android.media.AudioManager;
 
 import android.net.ConnectivityManager;
@@ -588,6 +590,11 @@ public class GameActivity extends $${gameActivitySuperClass}$$ implements Surfac
 	private int memstatLoggingInterval = 30000;
 	private int memstatRSSRequiredDiff = 2048;
 	
+	/** Audio focus listener */
+	private boolean bIsAudioFocusTracked = false;
+	private boolean bHasAudioFocusListeners = false;
+	private AudioFocusRequest audioFocusRequest = null;
+	
 	enum EAlertDialogType
 	{
 		None,
@@ -5030,6 +5037,62 @@ Trace.endSection();
 	{
 //$${gameActivityResizeKeyboardAdditions}$$
 	}
+	
+	public void AndroidThunkJava_OnAudioFocusChangedListenersUpdated(boolean bInHasListeners)
+	{
+		bHasAudioFocusListeners = bInHasListeners;
+		if (!bIsAudioFocusTracked && bHasAudioFocusListeners)
+		{
+			RequestAndTrackAudioFocus();
+		}
+		else if (bIsAudioFocusTracked && !bHasAudioFocusListeners)
+		{
+			AbandonAudioFocusRequest();
+		}
+	}
+    	
+	public void RequestAndTrackAudioFocus()
+	{
+		// Request audio focus and bind nativeOnAudioFocusChanged to audio focus changes
+		if (!isStandalone && Build.VERSION.SDK_INT >= 26)
+		{
+			AudioAttributes audioAttributes = new AudioAttributes.Builder()
+				.setUsage(AudioAttributes.USAGE_GAME)
+				.setContentType(AudioAttributes.CONTENT_TYPE_UNKNOWN)
+				.build();
+				
+			AudioManager.OnAudioFocusChangeListener audioFocusChangeListener = new AudioManager.OnAudioFocusChangeListener()
+			{
+				@Override
+				public void onAudioFocusChange(int focusChange)
+				{
+					if (bHasAudioFocusListeners)
+					{
+						nativeOnAudioFocusChanged(focusChange);
+					}
+				}
+			};
+			  
+			AudioManager audioManager = (AudioManager)getThisActivity().getSystemService(Context.AUDIO_SERVICE);
+			audioFocusRequest = new AudioFocusRequest.Builder(AudioManager.AUDIOFOCUS_GAIN)
+				.setAudioAttributes(audioAttributes)	
+				.setOnAudioFocusChangeListener(audioFocusChangeListener)
+				.build();
+			audioManager.requestAudioFocus(audioFocusRequest);
+			bIsAudioFocusTracked = true;
+		}
+	}
+
+	public void AbandonAudioFocusRequest()
+	{
+		if (!isStandalone && Build.VERSION.SDK_INT >= 26 && bIsAudioFocusTracked && audioFocusRequest != null)
+		{
+			AudioManager audioManager = (AudioManager)getThisActivity().getSystemService(Context.AUDIO_SERVICE);
+			audioManager.abandonAudioFocusRequest(audioFocusRequest);
+			audioFocusRequest = null;
+			bIsAudioFocusTracked = false;
+		}
+	}
 
 	@Override
 	public void onAccuracyChanged(Sensor sensor, int accuracy)
@@ -8896,6 +8959,8 @@ Trace.endSection();
 	public native void nativeNetworkChanged();
 
 	public native void nativeHandleNewIntentData(String intentUri);
+	
+	public native void nativeOnAudioFocusChanged(int focusChange);
 
 	@NonNull String GetNotShippingCommandLine(@NonNull Bundle Bundle, @NonNull String CommandLine)
 	{
diff --git a/Engine/Source/Runtime/Core/Private/Android/AndroidPlatformMisc.cpp b/Engine/Source/Runtime/Core/Private/Android/AndroidPlatformMisc.cpp
--- a/Engine/Source/Runtime/Core/Private/Android/AndroidPlatformMisc.cpp
+++ b/Engine/Source/Runtime/Core/Private/Android/AndroidPlatformMisc.cpp
@@ -1664,6 +1664,8 @@ extern void AndroidThunkCpp_RegisterForRemoteNotifications();
 extern void AndroidThunkCpp_UnregisterForRemoteNotifications();
 extern bool AndroidThunkCpp_IsAllowedRemoteNotifications();
 
+extern void AndroidThunkCpp_OnAudioFocusChangedListenersUpdated(bool bInHasListeners);
+
 void FAndroidMisc::RegisterForRemoteNotifications()
 {
 #if USE_ANDROID_JNI
@@ -1687,6 +1689,37 @@ bool FAndroidMisc::IsAllowedRemoteNotifications()
 #endif
 }
 
+FDelegateHandle FAndroidMisc::RegisterAudioFocusChangedEventListener(const FOnAudioFocusChangedDelegate Callback)
+{
+	if (!FCoreDelegates::AudioFocusChangedDelegate.IsBound())
+	{
+		AndroidThunkCpp_OnAudioFocusChangedListenersUpdated(true);
+	}
+	return FCoreDelegates::AudioFocusChangedDelegate.Add(Callback);
+}
+
+void FAndroidMisc::UnregisterAudioFocusChangedEventListener(const FDelegateHandle DelegateHandle)
+{
+	FCoreDelegates::AudioFocusChangedDelegate.Remove(DelegateHandle);
+	if (!FCoreDelegates::AudioFocusChangedDelegate.IsBound())
+	{
+		AndroidThunkCpp_OnAudioFocusChangedListenersUpdated(false);
+	}
+}
+
+#if USE_ANDROID_JNI
+void JNICALL UE::Jni::FGameActivity::nativeOnAudioFocusChanged(JNIEnv* jenv, jobject thiz, jint FocusChange)
+{
+	if (FTaskGraphInterface::IsRunning())
+	{
+		FFunctionGraphTask::CreateAndDispatchWhenReady([FocusChange]()
+		{
+			FCoreDelegates::AudioFocusChangedDelegate.Broadcast(static_cast<EAndroidAudioFocusChange>(FocusChange) == EAndroidAudioFocusChange::Gain);
+		}, TStatId(), NULL, ENamedThreads::GameThread);
+	}
+}
+#endif
+
 static FString PendingProtocolActivationUri;
 static void HandleNewIntentUri_GameThread()
 {
diff --git a/Engine/Source/Runtime/Core/Private/GenericPlatform/GenericPlatformMisc.cpp b/Engine/Source/Runtime/Core/Private/GenericPlatform/GenericPlatformMisc.cpp
--- a/Engine/Source/Runtime/Core/Private/GenericPlatform/GenericPlatformMisc.cpp
+++ b/Engine/Source/Runtime/Core/Private/GenericPlatform/GenericPlatformMisc.cpp
@@ -2247,6 +2247,17 @@ void FGenericPlatformMisc::UnregisterForRemoteNotifications()
 	// not implemented by default
 }
 
+FDelegateHandle FGenericPlatformMisc::RegisterAudioFocusChangedEventListener(const FOnAudioFocusChangedDelegate Callback)
+{
+	// not implemented by default
+	return FDelegateHandle();
+}
+
+void FGenericPlatformMisc::UnregisterAudioFocusChangedEventListener(const FDelegateHandle DelegateHandle)
+{
+	// not implemented by default
+}
+
 FString FGenericPlatformMisc::GetPendingActivationProtocol()
 {
 	// not implemented by default
diff --git a/Engine/Source/Runtime/Core/Private/Misc/CoreDelegates.cpp b/Engine/Source/Runtime/Core/Private/Misc/CoreDelegates.cpp
--- a/Engine/Source/Runtime/Core/Private/Misc/CoreDelegates.cpp
+++ b/Engine/Source/Runtime/Core/Private/Misc/CoreDelegates.cpp
@@ -186,6 +186,8 @@ TMulticastDelegate<void(FString, int, int)> FCoreDelegates::ApplicationReceivedL
 TMulticastDelegate<void()> FCoreDelegates::ApplicationPerformFetchDelegate;
 TMulticastDelegate<void(FString)> FCoreDelegates::ApplicationBackgroundSessionEventDelegate;
 
+TMulticastDelegate<void(bool)> FCoreDelegates::AudioFocusChangedDelegate;
+
 FSimpleMulticastDelegate& FCoreDelegates::GetOnPostEngineInit()
 {
 	static FSimpleMulticastDelegate Singleton;
diff --git a/Engine/Source/Runtime/Core/Public/Android/AndroidJniGameActivity.h b/Engine/Source/Runtime/Core/Public/Android/AndroidJniGameActivity.h
--- a/Engine/Source/Runtime/Core/Public/Android/AndroidJniGameActivity.h
+++ b/Engine/Source/Runtime/Core/Public/Android/AndroidJniGameActivity.h
@@ -74,6 +74,7 @@ namespace UE::Jni
 		static void JNICALL nativeSetWindowInfo(JNIEnv* env, jobject thiz, jboolean bIsPortrait, jint DepthBufferPreference, jint PropagateAlpha);
 		static void JNICALL nativeSetSurfaceViewInfo(JNIEnv* env, jobject thiz, jint width, jint height);
 		static void JNICALL nativeSetSafezoneInfo(JNIEnv* env, jobject thiz, jboolean bIsPortrait, jfloat left, jfloat top, jfloat right, jfloat bottom);
+		static void JNICALL nativeOnAudioFocusChanged(JNIEnv* env, jobject thiz, jint FocusChange);
 #if USE_ANDROID_INPUT
 
 		static constexpr FMember Members[]
@@ -121,7 +122,8 @@ namespace UE::Jni
 			UE_JNI_NATIVE_METHOD(nativeSupportsNEON),
 			UE_JNI_NATIVE_METHOD(nativeSetWindowInfo),
 			UE_JNI_NATIVE_METHOD(nativeSetSurfaceViewInfo),
-			UE_JNI_NATIVE_METHOD(nativeSetSafezoneInfo)
+			UE_JNI_NATIVE_METHOD(nativeSetSafezoneInfo),
+			UE_JNI_NATIVE_METHOD(nativeOnAudioFocusChanged)
 		};
 	};
 }
diff --git a/Engine/Source/Runtime/Core/Public/Android/AndroidPlatformMisc.h b/Engine/Source/Runtime/Core/Public/Android/AndroidPlatformMisc.h
--- a/Engine/Source/Runtime/Core/Public/Android/AndroidPlatformMisc.h
+++ b/Engine/Source/Runtime/Core/Public/Android/AndroidPlatformMisc.h
@@ -26,6 +26,15 @@ class TFunction;
 
 enum class ECrashContextType;
 
+enum class EAndroidAudioFocusChange
+{
+	LossTransientCanDuck = -3,
+	LossTransient = -2,
+	Loss = -1,
+	None = 0,
+	Gain = 1
+};
+
 /**
  * Android implementation of the misc OS functions
  */
@@ -159,6 +168,9 @@ public:
 	static CORE_API void RegisterForRemoteNotifications();
 	static CORE_API void UnregisterForRemoteNotifications();
 	static CORE_API bool IsAllowedRemoteNotifications();
+	
+	static CORE_API FDelegateHandle RegisterAudioFocusChangedEventListener(const FOnAudioFocusChangedDelegate Callback);
+	static CORE_API void UnregisterAudioFocusChangedEventListener(const FDelegateHandle DelegateHandle);
 
 	static CORE_API void HandleNewIntentUri(const FString& IntentUri);
 	static CORE_API FString GetPendingActivationProtocol();
diff --git a/Engine/Source/Runtime/Core/Public/GenericPlatform/GenericPlatformMisc.h b/Engine/Source/Runtime/Core/Public/GenericPlatform/GenericPlatformMisc.h
--- a/Engine/Source/Runtime/Core/Public/GenericPlatform/GenericPlatformMisc.h
+++ b/Engine/Source/Runtime/Core/Public/GenericPlatform/GenericPlatformMisc.h
@@ -563,6 +563,8 @@ namespace ECPUFeatureBits_X86
 	inline constexpr uint32 FMA = 1U << 12;
 }
 
+typedef TDelegate<void(bool), FDefaultDelegateUserPolicy> FOnAudioFocusChangedDelegate;
+
 /**
  * Generic implementation for most platforms
  */
@@ -1871,6 +1873,16 @@ public:
 	*/
 	static CORE_API void UnregisterForRemoteNotifications();
 
+	/**
+	* Register an audio focus change listener
+	*/
+	static CORE_API FDelegateHandle RegisterAudioFocusChangedEventListener(const FOnAudioFocusChangedDelegate Callback);
+
+	/**
+	* Unregister an audio focus change listener
+	*/
+	static CORE_API void UnregisterAudioFocusChangedEventListener(const FDelegateHandle DelegateHandle);
+
 	/**
 	* Returns activation protocol that's waiting to be broadcast
 	*/
diff --git a/Engine/Source/Runtime/Core/Public/Misc/CoreDelegates.h b/Engine/Source/Runtime/Core/Public/Misc/CoreDelegates.h
--- a/Engine/Source/Runtime/Core/Public/Misc/CoreDelegates.h
+++ b/Engine/Source/Runtime/Core/Public/Misc/CoreDelegates.h
@@ -576,6 +576,9 @@ public:
 	 */
 	static CORE_API TMulticastDelegate<void(ENetworkConnectionStatus LastConnectionState, ENetworkConnectionStatus ConnectionState)> OnNetworkConnectionStatusChanged;
 
+	// Called when the application audio focus changed. Only called if there's audio focus event listeners
+	static CORE_API TMulticastDelegate<void(bool bHasAudioFocus)> AudioFocusChangedDelegate;
+
 	// Callback to let code read or write specialized binary data that is generated at Stage time, for optimizing data right before 
 	// final game data is being written to disk
 	// The TMap is a map of an identifier for owner of the data, and a boolean where true means the data is being generated (ie editor), and false 
diff --git a/Engine/Source/Runtime/Launch/Private/Android/AndroidJNI.cpp b/Engine/Source/Runtime/Launch/Private/Android/AndroidJNI.cpp
--- a/Engine/Source/Runtime/Launch/Private/Android/AndroidJNI.cpp
+++ b/Engine/Source/Runtime/Launch/Private/Android/AndroidJNI.cpp
@@ -199,6 +199,9 @@ void FJavaWrapper::FindClassesAndMethods(JNIEnv* Env)
 	// motion controls
 	AndroidThunkJava_EnableMotion = FindMethod(Env, GameActivityClassID, "AndroidThunkJava_EnableMotion", "(Z)V", bIsOptional);
 
+	// audio focus
+	AndroidThunkJava_OnAudioFocusChangedListenersUpdated = FindMethod(Env, GameActivityClassID, "AndroidThunkJava_OnAudioFocusChangedListenersUpdated", "(Z)V", bIsOptional);
+
 	// the rest are optional
 	bIsOptional = true;
 
@@ -631,6 +634,8 @@ jmethodID FJavaWrapper::SetNameMethod;
 jmethodID FJavaWrapper::AndroidThunkJava_AcquireWifiManagerMulticastLock;
 jmethodID FJavaWrapper::AndroidThunkJava_ReleaseWifiManagerMulticastLock;
 
+jmethodID FJavaWrapper::AndroidThunkJava_OnAudioFocusChangedListenersUpdated;
+
 //Game-specific crash reporter
 void EngineCrashHandler(const FGenericCrashContext& GenericContext)
 {
@@ -1319,6 +1324,14 @@ void JNICALL UE::Jni::FGameActivity::nativeVirtualKeyboardSendSelection(JNIEnv*
 	}
 }
 
+void AndroidThunkCpp_OnAudioFocusChangedListenersUpdated(bool bInHasListeners)
+{
+	if (JNIEnv* Env = FAndroidApplication::GetJavaEnv())
+	{
+		FJavaWrapper::CallVoidMethod(Env, FJavaWrapper::GameActivityThis, FJavaWrapper::AndroidThunkJava_OnAudioFocusChangedListenersUpdated, bInHasListeners);
+	}
+}
+
 void AndroidThunkCpp_LaunchURL(const FString& URL)
 {
 	if (JNIEnv* Env = FAndroidApplication::GetJavaEnv())
diff --git a/Engine/Source/Runtime/Launch/Public/Android/AndroidJNI.h b/Engine/Source/Runtime/Launch/Public/Android/AndroidJNI.h
--- a/Engine/Source/Runtime/Launch/Public/Android/AndroidJNI.h
+++ b/Engine/Source/Runtime/Launch/Public/Android/AndroidJNI.h
@@ -167,6 +167,9 @@ public:
 	static jmethodID AndroidThunkJava_AcquireWifiManagerMulticastLock;
 	static jmethodID AndroidThunkJava_ReleaseWifiManagerMulticastLock;
 
+	// audio focus
+	static jmethodID AndroidThunkJava_OnAudioFocusChangedListenersUpdated;
+
 	/**
 	 * Find all known classes and methods
 	 */
